ASTC Encoder

diff --git a/Engine/Source/Developer/TextureFormatASTC/Private/TextureFormatASTC.cpp b/Engine/Source/Developer/TextureFormatASTC/Private/TextureFormatASTC.cpp
--- a/Engine/Source/Developer/TextureFormatASTC/Private/TextureFormatASTC.cpp
+++ b/Engine/Source/Developer/TextureFormatASTC/Private/TextureFormatASTC.cpp
@@ -115,7 +115,7 @@ static bool CompressSliceToASTC(
 	}
 
 	// Start Compressor
-	FString compressorPath(FPaths::EngineDir() + TEXT("Binaries/ThirdParty/ARM/astcenc.exe"));
+	FString compressorPath(FPaths::EngineDir() + TEXT("Binaries/ThirdParty/ARM/Win32/astcenc.exe"));
 	void* Proc = FPlatformProcess::CreateProc(compressorPath.GetCharArray().GetData(), *Params, true, false, false, NULL, -1, NULL, NULL);
 
 	// Failed to start the compressor process
