- Basic ASTC texture support for Metal (A8X and up) - this is experimental, and some textures are not working properly. Not exposed in the editor to enable (that will be a future changelist)

diff --git a/Engine/Config/BaseEngine.ini b/Engine/Config/BaseEngine.ini
--- a/Engine/Config/BaseEngine.ini
+++ b/Engine/Config/BaseEngine.ini
@@ -2355,9 +2355,10 @@ bSupportsPortraitOrientation=False
 bSupportsUpsideDownOrientation=False
 bSupportsLandscapeLeftOrientation=True
 bSupportsLandscapeRightOrientation=True
-bSupportsMetal=False
 bSupportsMetal=True
 bSupportsOpenGLES2=True
+bCookPVRTCTextures=True
+bCookASTCTextures=False
 
 ; These are the defaults for Android settings, and they need to be in the .ini since UBT reads the .ini settings, without instantiating the class
 [/Script/AndroidRuntimeSettings.AndroidRuntimeSettings]
@@ -2371,6 +2372,11 @@ bBuildForES31=false
 bBuildForGL4=false
 bSplitIntoSeparateApks=false
 
+[/Script/UnrealEd.CookerSettings]
+DefaultPVRTCQuality=1
+DefaultASTCQualityBySpeed=3
+DefaultASTCQualityBySize=3
+
 [/Script/Engine.PhysicsSettings]
 DefaultGravityZ=-980.0
 bEnable2DPhysics=false
diff --git a/Engine/Source/Developer/TargetPlatform/TargetPlatform.Build.cs b/Engine/Source/Developer/TargetPlatform/TargetPlatform.Build.cs
--- a/Engine/Source/Developer/TargetPlatform/TargetPlatform.Build.cs
+++ b/Engine/Source/Developer/TargetPlatform/TargetPlatform.Build.cs
@@ -40,6 +40,7 @@ public class TargetPlatform : ModuleRules
 				{
                     DynamicallyLoadedModuleNames.Add("TextureFormatDXT");
                     DynamicallyLoadedModuleNames.Add("TextureFormatPVR");
+					DynamicallyLoadedModuleNames.Add("TextureFormatASTC");
 				}
 
 				DynamicallyLoadedModuleNames.Add("TextureFormatUncompressed");
@@ -73,6 +74,7 @@ public class TargetPlatform : ModuleRules
 				{
 					DynamicallyLoadedModuleNames.Add("TextureFormatDXT");
                     DynamicallyLoadedModuleNames.Add("TextureFormatPVR");
+					DynamicallyLoadedModuleNames.Add("TextureFormatASTC");
 				}
 
 				DynamicallyLoadedModuleNames.Add("TextureFormatUncompressed");
@@ -103,6 +105,7 @@ public class TargetPlatform : ModuleRules
                 {
                     DynamicallyLoadedModuleNames.Add("TextureFormatDXT");
                     DynamicallyLoadedModuleNames.Add("TextureFormatPVR");
+					DynamicallyLoadedModuleNames.Add("TextureFormatASTC");
                 }
 
                 DynamicallyLoadedModuleNames.Add("TextureFormatUncompressed");
diff --git a/Engine/Source/Developer/TextureFormatASTC/Private/TextureFormatASTC.cpp b/Engine/Source/Developer/TextureFormatASTC/Private/TextureFormatASTC.cpp
--- a/Engine/Source/Developer/TextureFormatASTC/Private/TextureFormatASTC.cpp
+++ b/Engine/Source/Developer/TextureFormatASTC/Private/TextureFormatASTC.cpp
@@ -11,6 +11,12 @@
 #include "PixelFormat.h"
 #include "GenericPlatformProcess.h"
 
+
+#define MAX_QUALITY_BY_SIZE 4
+#define MAX_QUALITY_BY_SPEED 3
+#define FORCED_NORMAL_MAP_COMPRESSION_SIZE_VALUE 1
+
+
 DEFINE_LOG_CATEGORY_STATIC(LogTextureFormatASTC, Log, All);
 
 /**
@@ -58,6 +64,81 @@ DEFINE_LOG_CATEGORY_STATIC(LogTextureFormatASTC, Log, All);
 
 IImageWrapperModule& ImageWrapperModule = FModuleManager::LoadModuleChecked<IImageWrapperModule>(FName("ImageWrapper"));
 
+
+static int32 GetDefaultCompressionBySizeValue()
+{
+	// start at default quality, then lookup in .ini file
+	int32 CompressionModeValue = 0;
+	GConfig->GetInt(TEXT("/Script/UnrealEd.CookerSettings"), TEXT("DefaultASTCQualityBySize"), CompressionModeValue, GEngineIni);
+	
+	FParse::Value(FCommandLine::Get(), TEXT("-astcqualitybysize="), CompressionModeValue);
+	CompressionModeValue = FMath::Min<uint32>(CompressionModeValue, MAX_QUALITY_BY_SIZE);
+	
+	return CompressionModeValue;
+}
+
+static int32 GetDefaultCompressionBySpeedValue()
+{
+	// start at default quality, then lookup in .ini file
+	int32 CompressionModeValue = 0;
+	GConfig->GetInt(TEXT("/Script/UnrealEd.CookerSettings"), TEXT("DefaultASTCQualityBySpeed"), CompressionModeValue, GEngineIni);
+	
+	FParse::Value(FCommandLine::Get(), TEXT("-astcqualitybyspeed="), CompressionModeValue);
+	CompressionModeValue = FMath::Min<uint32>(CompressionModeValue, MAX_QUALITY_BY_SPEED);
+	
+	return CompressionModeValue;
+}
+
+static FString GetQualityString(int32 OverrideSizeValue=-1, int32 OverrideSpeedValue=-1)
+{
+	// convert to a string
+	FString CompressionMode;
+	switch (OverrideSizeValue >= 0 ? OverrideSizeValue : GetDefaultCompressionBySizeValue())
+	{
+		case 0:	CompressionMode = TEXT("12x12"); break;
+		case 1:	CompressionMode = TEXT("10x10"); break;
+		case 2:	CompressionMode = TEXT("8x8"); break;
+		case 3:	CompressionMode = TEXT("6x6"); break;
+		case 4:	CompressionMode = TEXT("4x4"); break;
+		default: UE_LOG(LogTemp, Fatal, TEXT("ASTC size quality higher than expected"));
+	}
+	
+	switch (OverrideSpeedValue >= 0 ? OverrideSpeedValue : GetDefaultCompressionBySpeedValue())
+	{
+		case 0:	CompressionMode += TEXT(" -veryfast"); break;
+		case 1:	CompressionMode += TEXT(" -fast"); break;
+		case 2:	CompressionMode += TEXT(" -medium"); break;
+		case 3:	CompressionMode += TEXT(" -thorough"); break;
+		default: UE_LOG(LogTemp, Fatal, TEXT("ASTC speed quality higher than expected"));
+	}
+
+	return CompressionMode;
+}
+
+static EPixelFormat GetQualityFormat(int32 OverrideSizeValue=-1)
+{
+	// convert to a string
+	EPixelFormat Format = PF_Unknown;
+	switch (OverrideSizeValue >= 0 ? OverrideSizeValue : GetDefaultCompressionBySizeValue())
+	{
+		case 0:	Format = PF_ASTC_12x12; break;
+		case 1:	Format = PF_ASTC_10x10; break;
+		case 2:	Format = PF_ASTC_8x8; break;
+		case 3:	Format = PF_ASTC_6x6; break;
+		case 4:	Format = PF_ASTC_4x4; break;
+		default: UE_LOG(LogTemp, Fatal, TEXT("Max quality higher than expected"));
+	}
+	
+	return Format;
+}
+
+static uint16 GetQualityVersion()
+{
+	// top 3 bits for size compression value, and next 3 for speed
+	return (GetDefaultCompressionBySizeValue() << 13) | (GetDefaultCompressionBySpeedValue() << 10);
+}
+
+
 static bool CompressSliceToASTC(
 	const void* SourceData,
 	int32 SizeX,
@@ -80,7 +161,7 @@ static bool CompressSliceToASTC(
 	
 	// Compress and retrieve the PNG data to write out to disk
 	IImageWrapperPtr ImageWrapper = ImageWrapperModule.CreateImageWrapper(EImageFormat::PNG);
-	ImageWrapper->SetRaw(SourceData, SizeX * SizeY * 4, SizeX, SizeY);
+	ImageWrapper->SetRaw(SourceData, SizeX * SizeY * 4, SizeX, SizeY, ERGBFormat::RGBA, 8);
 	const TArray<uint8>& FileData = ImageWrapper->GetCompressed();
 	int32 FileDataSize = FileData.Num();
 
@@ -95,33 +176,37 @@ static bool CompressSliceToASTC(
 	FArchive* PNGFile = NULL;
 	while (!PNGFile)
 	{
-		PNGFile = GFileManager->CreateFileWriter(*InputFilePath);   // Occasionally returns NULL due to error code ERROR_SHARING_VIOLATION
+		PNGFile = IFileManager::Get().CreateFileWriter(*InputFilePath);   // Occasionally returns NULL due to error code ERROR_SHARING_VIOLATION
 		FPlatformProcess::Sleep(0.01f);                             // ... no choice but to wait for the file to become free to access
 	}
 	PNGFile->Serialize((void*)&FileData[0], FileDataSize);
 	delete PNGFile;
 
 	// Compress PNG file to ASTC (using the reference astcenc.exe from ARM)
-	FString Params = FString::Printf(TEXT("-c %s %s %s -thorough"),
+	FString Params = FString::Printf(TEXT("-c %s %s %s"),
 		*InputFilePath,
 		*OutputFilePath,
 		*CompressionParameters
 	);
 
-	// Give a debug message about the process
-	if (GIsUCC)
-	{
-		//UE_LOG(LogTextureFormatASTC, Display, TEXT("Compressing to ASTC_%s..."), *CompressionRate);
-	}
+	UE_LOG(LogTextureFormatASTC, Display, TEXT("Compressing to ASTC (%s)..."), *CompressionParameters);
 
 	// Start Compressor
-	FString compressorPath(FPaths::EngineDir() + TEXT("Binaries/ThirdParty/ARM/Win32/astcenc.exe"));
-	void* Proc = FPlatformProcess::CreateProc(compressorPath.GetCharArray().GetData(), *Params, true, false, false, NULL, -1, NULL, NULL);
+#if PLATFORM_MAC
+	FString CompressorPath(FPaths::EngineDir() + TEXT("Binaries/ThirdParty/ARM/Mac/astcenc"));
+#elif PLATFORM_LINUX
+	FString CompressorPath(FPaths::EngineDir() + TEXT("Binaries/ThirdParty/ARM/Linux32/astcenc"));
+#elif PLATFORM_WINDOWS
+	FString CompressorPath(FPaths::EngineDir() + TEXT("Binaries/ThirdParty/ARM/Win32/astcenc.exe"));
+#else
+#error Unsupported platform
+#endif
+	FProcHandle Proc = FPlatformProcess::CreateProc(*CompressorPath, *Params, true, false, false, NULL, -1, NULL, NULL);
 
 	// Failed to start the compressor process
-	if (Proc == NULL)
+	if (!Proc.IsValid())
 	{
-		UE_LOG(LogTextureFormatASTC, Error, TEXT("Failed to start astcenc.exe for compressing images (%s)"), compressorPath.GetCharArray().GetData());
+		UE_LOG(LogTextureFormatASTC, Error, TEXT("Failed to start astcenc for compressing images (%s)"), *CompressorPath);
 		return false;
 	}
 
@@ -159,14 +244,14 @@ static bool CompressSliceToASTC(
 			(Header->TexelCountZ[1] <<  8) + 
 			(Header->TexelCountZ[2] << 16);
 
-		//UE_LOG(LogTextureFormatASTC, Display, TEXT("    Compressed Texture Header:"));
-		//UE_LOG(LogTextureFormatASTC, Display, TEXT("             Magic: %x"), Header->Magic);
-		//UE_LOG(LogTextureFormatASTC, Display, TEXT("        BlockSizeX: %u"), Header->BlockSizeX);
-		//UE_LOG(LogTextureFormatASTC, Display, TEXT("        BlockSizeY: %u"), Header->BlockSizeY);
-		//UE_LOG(LogTextureFormatASTC, Display, TEXT("        BlockSizeZ: %u"), Header->BlockSizeZ);
-		//UE_LOG(LogTextureFormatASTC, Display, TEXT("       TexelCountX: %u"), TexelCountX);
-		//UE_LOG(LogTextureFormatASTC, Display, TEXT("       TexelCountY: %u"), TexelCountY);
-		//UE_LOG(LogTextureFormatASTC, Display, TEXT("       TexelCountZ: %u"), TexelCountZ);
+//		UE_LOG(LogTextureFormatASTC, Display, TEXT("    Compressed Texture Header:"));
+//		UE_LOG(LogTextureFormatASTC, Display, TEXT("             Magic: %x"), Header->Magic);
+//		UE_LOG(LogTextureFormatASTC, Display, TEXT("        BlockSizeX: %u"), Header->BlockSizeX);
+//		UE_LOG(LogTextureFormatASTC, Display, TEXT("        BlockSizeY: %u"), Header->BlockSizeY);
+//		UE_LOG(LogTextureFormatASTC, Display, TEXT("        BlockSizeZ: %u"), Header->BlockSizeZ);
+//		UE_LOG(LogTextureFormatASTC, Display, TEXT("       TexelCountX: %u"), TexelCountX);
+//		UE_LOG(LogTextureFormatASTC, Display, TEXT("       TexelCountY: %u"), TexelCountY);
+//		UE_LOG(LogTextureFormatASTC, Display, TEXT("       TexelCountZ: %u"), TexelCountZ);
 
 		// Calculate size of this mip in blocks
 		uint32 MipSizeX = (TexelCountX + Header->BlockSizeX - 1) / Header->BlockSizeX;
@@ -178,7 +263,7 @@ static bool CompressSliceToASTC(
 		// Copy the compressed data
 		OutCompressedData.Empty(MipSize);
 		OutCompressedData.AddUninitialized(MipSize);
-		void* MipData = OutCompressedData.GetTypedData();
+		void* MipData = OutCompressedData.GetData();
 
 		// Calculate the offset to get to the mip data
 		check(sizeof(FASTCHeader) == 16);
@@ -188,14 +273,14 @@ static bool CompressSliceToASTC(
 	else
 	{
 		UE_LOG(LogTextureFormatASTC, Error, TEXT("ASTC encoder failed with return code %d, mip size (%d, %d)"), ReturnCode, SizeX, SizeY);
-		GFileManager->Delete(*InputFilePath);
-		GFileManager->Delete(*OutputFilePath);
+		IFileManager::Get().Delete(*InputFilePath);
+		IFileManager::Get().Delete(*OutputFilePath);
 		return false;
 	}
 		
 	// Delete intermediate files
-	GFileManager->Delete(*InputFilePath);
-	GFileManager->Delete(*OutputFilePath);
+	IFileManager::Get().Delete(*InputFilePath);
+	IFileManager::Get().Delete(*OutputFilePath);
 	return true;
 }
 
@@ -211,7 +296,14 @@ class FTextureFormatASTC : public ITextureFormat
 
 	virtual uint16 GetVersion(FName Format) const OVERRIDE
 	{
-		return 27;
+		return GetQualityVersion() + 29;
+	}
+
+	virtual FTextureFormatCompressorCaps GetFormatCapabilities() const override
+	{
+		FTextureFormatCompressorCaps RetCaps;
+		// use defaults
+		return RetCaps;
 	}
 
 	virtual void GetSupportedFormats(TArray<FName>& OutFormats) const OVERRIDE
@@ -240,24 +332,24 @@ class FTextureFormatASTC : public ITextureFormat
 		if (BuildSettings.TextureFormatName == GTextureFormatNameASTC_RGB ||
 		  ((BuildSettings.TextureFormatName == GTextureFormatNameASTC_RGBAuto) && !bImageHasAlphaChannel))
 		{
-			CompressedPixelFormat = PF_ASTC_8x8;
-			CompressionParameters = TEXT("8x8 -esw bgra -ch 1 1 1 0");
+			CompressedPixelFormat = GetQualityFormat();
+			CompressionParameters = FString::Printf(TEXT("%s -esw bgra -ch 1 1 1 0"), *GetQualityString());
 		}
 		else if (BuildSettings.TextureFormatName == GTextureFormatNameASTC_RGBA ||
 			   ((BuildSettings.TextureFormatName == GTextureFormatNameASTC_RGBAuto) && bImageHasAlphaChannel))
 		{
-			CompressedPixelFormat = PF_ASTC_6x6;
-			CompressionParameters = TEXT("6x6 -esw bgra -ch 1 1 1 1 -alphablend");
+			CompressedPixelFormat = GetQualityFormat();
+			CompressionParameters = FString::Printf(TEXT("%s -esw bgra -ch 1 1 1 1 -alphablend"), *GetQualityString());
 		}
 		else if (BuildSettings.TextureFormatName == GTextureFormatNameASTC_NormalAG)
 		{
-			CompressedPixelFormat = PF_ASTC_6x6;
-			CompressionParameters = TEXT("6x6 -esw 0g0b -ch 0 1 0 1 -oplimit 1000 -mincorrel 0.99 -dblimit 60 -b 2.5 -v 3 1 1 0 50 0 -va 1 1 0 50");
+			CompressedPixelFormat = GetQualityFormat(FORCED_NORMAL_MAP_COMPRESSION_SIZE_VALUE);
+			CompressionParameters = FString::Printf(TEXT("%s -esw 0g0b -ch 0 1 0 1 -oplimit 1000 -mincorrel 0.99 -dblimit 60 -b 2.5 -v 3 1 1 0 50 0 -va 1 1 0 50"), *GetQualityString(FORCED_NORMAL_MAP_COMPRESSION_SIZE_VALUE, -1));
 		}
 		else if (BuildSettings.TextureFormatName == GTextureFormatNameASTC_NormalRG)
 		{
-			CompressedPixelFormat = PF_ASTC_6x6;
-			CompressionParameters = TEXT("6x6 -esw bg00 -ch 1 1 0 0 -oplimit 1000 -mincorrel 0.99 -dblimit 60 -b 2.5 -v 3 1 1 0 50 0 -va 1 1 0 50");
+			CompressedPixelFormat = GetQualityFormat(FORCED_NORMAL_MAP_COMPRESSION_SIZE_VALUE);
+			CompressionParameters = FString::Printf(TEXT("%s -esw bg00 -ch 1 1 0 0 -oplimit 1000 -mincorrel 0.99 -dblimit 60 -b 2.5 -v 3 1 1 0 50 0 -va 1 1 0 50"), *GetQualityString(FORCED_NORMAL_MAP_COMPRESSION_SIZE_VALUE, -1));
 		}
 
 		// Compress the image, slice by slice
diff --git a/Engine/Source/Developer/TextureFormatPVR/Private/TextureFormatPVR.cpp b/Engine/Source/Developer/TextureFormatPVR/Private/TextureFormatPVR.cpp
--- a/Engine/Source/Developer/TextureFormatPVR/Private/TextureFormatPVR.cpp
+++ b/Engine/Source/Developer/TextureFormatPVR/Private/TextureFormatPVR.cpp
@@ -11,7 +11,6 @@
 #include "GenericPlatformProcess.h"
 
 
-#define DEFAULT_QUALITY 0
 #define MAX_QUALITY 4
 
 
@@ -227,18 +226,23 @@ static void UseOriginal(const FImage& InImage, FCompressedImage2D& OutCompressed
     FMemory::Memcpy(MipData, Image.RawData.GetData(), Image.SizeX * Image.SizeY * 4);
 }
 
-static FString GetPVRTCQualityString()
+static int32 GetDefaultCompressionValue()
 {
 	// start at default quality, then lookup in .ini file
-	int32 CompressionModeValue = DEFAULT_QUALITY;
+	int32 CompressionModeValue = 0;
 	GConfig->GetInt(TEXT("/Script/UnrealEd.CookerSettings"), TEXT("DefaultPVRTCQuality"), CompressionModeValue, GEngineIni);
 
 	FParse::Value(FCommandLine::Get(), TEXT("-pvrtcquality="), CompressionModeValue);
 	CompressionModeValue = FMath::Min<uint32>(CompressionModeValue, MAX_QUALITY);
+	
+	return CompressionModeValue;
+}
 
+static FString GetPVRTCQualityString()
+{
 	// convert to a string
 	FString CompressionMode;
-	switch (CompressionModeValue)
+	switch (GetDefaultCompressionValue())
 	{
 		case 0:	CompressionMode = TEXT("fastest"); break;
 		case 1:	CompressionMode = TEXT("fast"); break;
@@ -253,15 +257,8 @@ static FString GetPVRTCQualityString()
 
 static uint16 GetPVRTCQualityForVersion()
 {
-	// start at default quality
-	int32 CompressionModeValue = DEFAULT_QUALITY;
-	GConfig->GetInt(TEXT("/Script/UnrealEd.CookerSettings"), TEXT("DefaultPVRTCQuality"), CompressionModeValue, GEngineIni);
-
-	FParse::Value(FCommandLine::Get(), TEXT("-pvrtcquality="), CompressionModeValue);
-	CompressionModeValue = FMath::Min<uint32>(CompressionModeValue, MAX_QUALITY);
-
 	// top 3 bits for compression value
-	return CompressionModeValue << 13;
+	return GetDefaultCompressionValue() << 13;
 }
 
 /**
diff --git a/Engine/Source/Developer/iOS/IOSTargetPlatform/Private/IOSTargetPlatform.cpp b/Engine/Source/Developer/iOS/IOSTargetPlatform/Private/IOSTargetPlatform.cpp
--- a/Engine/Source/Developer/iOS/IOSTargetPlatform/Private/IOSTargetPlatform.cpp
+++ b/Engine/Source/Developer/iOS/IOSTargetPlatform/Private/IOSTargetPlatform.cpp
@@ -314,6 +314,23 @@ static bool SupportsMetalMRT()
 	return bSupportsMetalMRT;
 }
 
+static bool CookPVRTC()
+{
+	// default to using PVRTC
+	bool bCookPVRTCTextures = true;
+	GConfig->GetBool(TEXT("/Script/IOSRuntimeSettings.IOSRuntimeSettings"), TEXT("bCookPVRTCTextures"), bCookPVRTCTextures, GEngineIni);
+	return bCookPVRTCTextures;
+}
+
+static bool CookASTC()
+{
+	// default to not using ASTC
+	bool bCookASTCTextures = true;
+	GConfig->GetBool(TEXT("/Script/IOSRuntimeSettings.IOSRuntimeSettings"), TEXT("bCookASTCTextures"), bCookASTCTextures, GEngineIni);
+	return bCookASTCTextures;
+}
+
+
 void FIOSTargetPlatform::GetAllPossibleShaderFormats( TArray<FName>& OutFormats ) const
 {
 	static FName NAME_OPENGL_ES2_IOS(TEXT("GLSL_ES2_IOS"));
@@ -346,14 +363,15 @@ void FIOSTargetPlatform::GetTextureFormats( const UTexture* Texture, TArray<FNam
 {
 	check(Texture);
 
-	// we remap some of the defaults
-	static FName FormatRemap[] = 
+	// we remap some of the defaults (with PVRTC and ASTC formats)
+	static FName FormatRemap[] =
 	{
-		FName(TEXT("DXT1")), FName(TEXT("PVRTC2")),
-		FName(TEXT("DXT5")), FName(TEXT("PVRTC4")),
-		FName(TEXT("DXT5n")), FName(TEXT("PVRTCN")),
-		FName(TEXT("BC5")), FName(TEXT("PVRTCN")),
-		FName(TEXT("AutoDXT")), FName(TEXT("AutoPVRTC")),
+		// original				PVRTC						ASTC
+		FName(TEXT("DXT1")),	FName(TEXT("PVRTC2")),		FName(TEXT("ASTC_RGB")),
+		FName(TEXT("DXT5")),	FName(TEXT("PVRTC4")),		FName(TEXT("ASTC_RGBA")),
+		FName(TEXT("DXT5n")),	FName(TEXT("PVRTCN")),		FName(TEXT("ASTC_NormalAG")),
+		FName(TEXT("BC5")),		FName(TEXT("PVRTCN")),		FName(TEXT("ASTC_NormalRG")),
+		FName(TEXT("AutoDXT")),	FName(TEXT("AutoPVRTC")),	FName(TEXT("ASTC_RGBAuto")),
 	};
 
 	FName TextureFormatName = NAME_None;
@@ -371,15 +389,32 @@ void FIOSTargetPlatform::GetTextureFormats( const UTexture* Texture, TArray<FNam
 	}
 
 	// perform any remapping away from defaults
-	for (int32 RemapIndex = 0; RemapIndex < ARRAY_COUNT(FormatRemap); RemapIndex += 2)
+	bool bFoundRemap = false;
+	bool bIncludePVRTC = CookPVRTC();
+	bool bIncludeASTC = CookASTC();
+	for (int32 RemapIndex = 0; RemapIndex < ARRAY_COUNT(FormatRemap); RemapIndex += 3)
 	{
 		if (TextureFormatName == FormatRemap[RemapIndex])
 		{
-			TextureFormatName = FormatRemap[RemapIndex + 1];
+			// we found a remapping
+			bFoundRemap = true;
+			// include the formats we want (use ASTC first so that it is preferred at runtime if they both exist and it's supported)
+			if (bIncludeASTC)
+			{
+				OutFormats.AddUnique(FormatRemap[RemapIndex + 2]);
+			}
+			if (bIncludePVRTC)
+			{
+				OutFormats.AddUnique(FormatRemap[RemapIndex + 1]);
+			}
 		}
 	}
 
-	OutFormats.Add(TextureFormatName);
+	// if we didn't already remap above, add it now
+	if (!bFoundRemap)
+	{
+		OutFormats.Add(TextureFormatName);
+	}
 }
 
 
diff --git a/Engine/Source/Editor/UnrealEd/Private/Commandlets/CookCommandlet.cpp b/Engine/Source/Editor/UnrealEd/Private/Commandlets/CookCommandlet.cpp
--- a/Engine/Source/Editor/UnrealEd/Private/Commandlets/CookCommandlet.cpp
+++ b/Engine/Source/Editor/UnrealEd/Private/Commandlets/CookCommandlet.cpp
@@ -31,6 +31,8 @@ UCookerSettings::UCookerSettings(const FObjectInitializer& ObjectInitializer)
 	: Super(ObjectInitializer)
 {
 	DefaultPVRTCQuality = 1;
+	DefaultASTCQualityBySize = 3;
+	DefaultASTCQualityBySpeed = 3;
 }
 
 
diff --git a/Engine/Source/Editor/UnrealEd/Public/CookerSettings.h b/Engine/Source/Editor/UnrealEd/Public/CookerSettings.h
--- a/Engine/Source/Editor/UnrealEd/Public/CookerSettings.h
+++ b/Engine/Source/Editor/UnrealEd/Public/CookerSettings.h
@@ -19,7 +19,15 @@ class UNREALED_API UCookerSettings
 
 public:
 
-	/** Quality of 0 means fastest, 4 means best */
-	UPROPERTY(config, EditAnywhere, Category=Textures)
+	/** Quality of 0 means fastest, 4 means best quality */
+	UPROPERTY(GlobalConfig, EditAnywhere, Category = Textures, meta = (DisplayName = "PVRTC Compression Quality (0-4, 0 is fastest)"))
 	int32 DefaultPVRTCQuality;
+
+	/** Quality of 0 means fastest, 3 means best quality */
+	UPROPERTY(GlobalConfig, EditAnywhere, Category = Textures, meta = (DisplayName = "ASTC Compression Quality vs Speed (0-4, 0 is fastest)"))
+	int32 DefaultASTCQualityBySpeed;
+
+	/** Quality of 0 means smallest (12x12 block size), 4 means best (4x4 block size) */
+	UPROPERTY(GlobalConfig, EditAnywhere, Category = Textures, meta = (DisplayName = "ASTC Compression Quality vs Size (0-4, 0 is smallest)"))
+	int32 DefaultASTCQualityBySize;
 };
diff --git a/Engine/Source/Programs/UnrealBuildTool/Android/UEBuildAndroid.cs b/Engine/Source/Programs/UnrealBuildTool/Android/UEBuildAndroid.cs
--- a/Engine/Source/Programs/UnrealBuildTool/Android/UEBuildAndroid.cs
+++ b/Engine/Source/Programs/UnrealBuildTool/Android/UEBuildAndroid.cs
@@ -249,6 +249,7 @@ namespace UnrealBuildTool
 						bBuildShaderFormats = true;
                         InModule.AddDynamicallyLoadedModule("TextureFormatPVR");
 						InModule.AddDynamicallyLoadedModule("TextureFormatDXT");
+						InModule.AddDynamicallyLoadedModule("TextureFormatASTC");
                         InModule.AddPlatformSpecificDynamicallyLoadedModule("TextureFormatAndroid");    // ATITC, ETC1 and ETC2
                         if (UEBuildConfiguration.bBuildDeveloperTools)
 						{
diff --git a/Engine/Source/Programs/UnrealBuildTool/IOS/UEBuildIOS.cs b/Engine/Source/Programs/UnrealBuildTool/IOS/UEBuildIOS.cs
--- a/Engine/Source/Programs/UnrealBuildTool/IOS/UEBuildIOS.cs
+++ b/Engine/Source/Programs/UnrealBuildTool/IOS/UEBuildIOS.cs
@@ -250,7 +250,8 @@ namespace UnrealBuildTool
                     else if (InModule.ToString() == "TargetPlatform")
                     {
                         bBuildShaderFormats = true;
-                        InModule.AddDynamicallyLoadedModule("TextureFormatPVR");
+						InModule.AddDynamicallyLoadedModule("TextureFormatPVR");
+						InModule.AddDynamicallyLoadedModule("TextureFormatASTC");
                         if (UEBuildConfiguration.bBuildDeveloperTools)
                         {
                             InModule.AddPlatformSpecificDynamicallyLoadedModule("AudioFormatADPCM");
diff --git a/Engine/Source/Runtime/IOS/MetalRHI/Private/MetalRHI.cpp b/Engine/Source/Runtime/IOS/MetalRHI/Private/MetalRHI.cpp
--- a/Engine/Source/Runtime/IOS/MetalRHI/Private/MetalRHI.cpp
+++ b/Engine/Source/Runtime/IOS/MetalRHI/Private/MetalRHI.cpp
@@ -46,6 +46,7 @@ FMetalDynamicRHI::FMetalDynamicRHI()
 	id<MTLDevice> Device = [IOSAppDelegate GetDelegate].IOSView->MetalDevice;
 	// A8 can use 256 bits of MRTs
 	bool bCanUseWideMRTs = [Device supportsFeatureSet:MTLFeatureSet_iOS_GPUFamily2_v1];
+	bool bCanUseASTC = [Device supportsFeatureSet:MTLFeatureSet_iOS_GPUFamily2_v1] && !FParse::Param(FCommandLine::Get(),TEXT("noastc"));
 
 	bool bProjectSupportsMRTs = false;
 	GConfig->GetBool(TEXT("/Script/IOSRuntimeSettings.IOSRuntimeSettings"), TEXT("bSupportsMetalMRT"), bProjectSupportsMRTs, GEngineIni);
@@ -113,6 +114,18 @@ FMetalDynamicRHI::FMetalDynamicRHI()
 	GPixelFormats[PF_PVRTC2				].Supported			= true;
 	GPixelFormats[PF_PVRTC4				].PlatformFormat	= MTLPixelFormatPVRTC_RGBA_4BPP;
 	GPixelFormats[PF_PVRTC4				].Supported			= true;
+	GPixelFormats[PF_PVRTC4				].PlatformFormat	= MTLPixelFormatPVRTC_RGBA_4BPP;
+	GPixelFormats[PF_PVRTC4				].Supported			= true;
+	GPixelFormats[PF_ASTC_4x4			].PlatformFormat	= MTLPixelFormatASTC_4x4_sRGB;
+	GPixelFormats[PF_ASTC_4x4			].Supported			= bCanUseASTC;
+	GPixelFormats[PF_ASTC_6x6			].PlatformFormat	= MTLPixelFormatASTC_6x6_sRGB;
+	GPixelFormats[PF_ASTC_6x6			].Supported			= bCanUseASTC;
+	GPixelFormats[PF_ASTC_8x8			].PlatformFormat	= MTLPixelFormatASTC_8x8_sRGB;
+	GPixelFormats[PF_ASTC_8x8			].Supported			= bCanUseASTC;
+	GPixelFormats[PF_ASTC_10x10			].PlatformFormat	= MTLPixelFormatASTC_10x10_sRGB;
+	GPixelFormats[PF_ASTC_10x10			].Supported			= bCanUseASTC;
+	GPixelFormats[PF_ASTC_12x12			].PlatformFormat	= MTLPixelFormatASTC_12x12_sRGB;
+	GPixelFormats[PF_ASTC_12x12			].Supported			= bCanUseASTC;
 	GPixelFormats[PF_UYVY				].PlatformFormat	= MTLPixelFormatInvalid;
 	GPixelFormats[PF_FloatRGB			].PlatformFormat	= MTLPixelFormatRGBA16Float;
 	GPixelFormats[PF_FloatRGB			].BlockBytes		= 8;
