Replace deprecated OVERRIDE with override

diff --git a/Engine/Plugins/Media/AndroidMedia/Source/AndroidMedia/Private/AndroidMediaModule.cpp b/Engine/Plugins/Media/AndroidMedia/Source/AndroidMedia/Private/AndroidMediaModule.cpp
--- a/Engine/Plugins/Media/AndroidMedia/Source/AndroidMedia/Private/AndroidMediaModule.cpp
+++ b/Engine/Plugins/Media/AndroidMedia/Source/AndroidMedia/Private/AndroidMediaModule.cpp
@@ -14,7 +14,7 @@ public:
 
 	// IModuleInterface interface
 
-	virtual void StartupModule() OVERRIDE
+	virtual void StartupModule() override
 	{
 		if (IsSupported())
 		{
@@ -29,7 +29,7 @@ public:
 		}
 	}
 
-	virtual void ShutdownModule() OVERRIDE
+	virtual void ShutdownModule() override
 	{
 		if (IsSupported())
 		{
diff --git a/Engine/Source/Developer/TextureFormatASTC/Private/TextureFormatASTC.cpp b/Engine/Source/Developer/TextureFormatASTC/Private/TextureFormatASTC.cpp
--- a/Engine/Source/Developer/TextureFormatASTC/Private/TextureFormatASTC.cpp
+++ b/Engine/Source/Developer/TextureFormatASTC/Private/TextureFormatASTC.cpp
@@ -289,12 +289,12 @@ static bool CompressSliceToASTC(
  */
 class FTextureFormatASTC : public ITextureFormat
 {
-	virtual bool AllowParallelBuild() const OVERRIDE
+	virtual bool AllowParallelBuild() const override
 	{
 		return true;
 	}
 
-	virtual uint16 GetVersion(FName Format) const OVERRIDE
+	virtual uint16 GetVersion(FName Format) const override
 	{
 		return GetQualityVersion() + 29;
 	}
@@ -306,7 +306,7 @@ class FTextureFormatASTC : public ITextureFormat
 		return RetCaps;
 	}
 
-	virtual void GetSupportedFormats(TArray<FName>& OutFormats) const OVERRIDE
+	virtual void GetSupportedFormats(TArray<FName>& OutFormats) const override
 	{
 		for (int32 i = 0; i < ARRAY_COUNT(GSupportedTextureFormatNames); ++i)
 		{
@@ -319,7 +319,7 @@ class FTextureFormatASTC : public ITextureFormat
 			const struct FTextureBuildSettings& BuildSettings,
 			bool bImageHasAlphaChannel,
 			FCompressedImage2D& OutCompressedImage
-		) const OVERRIDE
+		) const override
 	{
 		// Get Raw Image Data from passed in FImage
 		FImage Image;
