- Fixed ASTC textures on Metal, except for one blue texture in Rivalry

diff --git a/Engine/Config/BaseEngine.ini b/Engine/Config/BaseEngine.ini
--- a/Engine/Config/BaseEngine.ini
+++ b/Engine/Config/BaseEngine.ini
@@ -2092,7 +2092,7 @@ bSplitIntoSeparateApks=false
 
 [/Script/UnrealEd.CookerSettings]
 DefaultPVRTCQuality=1
-DefaultASTCQualityBySpeed=3
+DefaultASTCQualityBySpeed=1
 DefaultASTCQualityBySize=3
 
 [/Script/Engine.PhysicsSettings]
diff --git a/Engine/Source/Developer/TextureFormatASTC/Private/TextureFormatASTC.cpp b/Engine/Source/Developer/TextureFormatASTC/Private/TextureFormatASTC.cpp
--- a/Engine/Source/Developer/TextureFormatASTC/Private/TextureFormatASTC.cpp
+++ b/Engine/Source/Developer/TextureFormatASTC/Private/TextureFormatASTC.cpp
@@ -12,9 +12,12 @@
 #include "GenericPlatformProcess.h"
 
 
+// increment this if you change anything that will affect compression in this file, including FORCED_NORMAL_MAP_COMPRESSION_SIZE_VALUE
+#define BASE_FORMAT_VERSION 30
+
 #define MAX_QUALITY_BY_SIZE 4
 #define MAX_QUALITY_BY_SPEED 3
-#define FORCED_NORMAL_MAP_COMPRESSION_SIZE_VALUE 1
+#define FORCED_NORMAL_MAP_COMPRESSION_SIZE_VALUE 4
 
 
 DEFINE_LOG_CATEGORY_STATIC(LogTextureFormatASTC, Log, All);
@@ -296,9 +299,17 @@ class FTextureFormatASTC : public ITextureFormat
 
 	virtual uint16 GetVersion(FName Format) const override
 	{
-		return GetQualityVersion() + 29;
+		return GetQualityVersion() + BASE_FORMAT_VERSION;
 	}
 
+//	// Since we want to have per texture [group] compression settings, we need to have the key based on the texture
+//	virtual FString GetDerivedDataKeyString(const class UTexture& Texture) const override
+//	{
+//		const int32 LODBias = GSystemSettings.TextureLODSettings.CalculateLODBias(Texture.Source.GetSizeX(), Texture.Source.GetSizeY(), Texture.LODGroup, Texture.LODBias, Texture.NumCinematicMipLevels, Texture.MipGenSettings);
+//		check(LODBias >= 0);
+//		return FString::Printf(TEXT("%02u%d_"), (uint32)LODBias, CVarVirtualTextureReducedMemoryEnabled->GetValueOnGameThread());
+//	}
+	
 	virtual FTextureFormatCompressorCaps GetFormatCapabilities() const override
 	{
 		FTextureFormatCompressorCaps RetCaps;
diff --git a/Engine/Source/Runtime/IOS/MetalRHI/Private/MetalRHI.cpp b/Engine/Source/Runtime/IOS/MetalRHI/Private/MetalRHI.cpp
--- a/Engine/Source/Runtime/IOS/MetalRHI/Private/MetalRHI.cpp
+++ b/Engine/Source/Runtime/IOS/MetalRHI/Private/MetalRHI.cpp
@@ -116,15 +116,15 @@ FMetalDynamicRHI::FMetalDynamicRHI()
 	GPixelFormats[PF_PVRTC4				].Supported			= true;
 	GPixelFormats[PF_PVRTC4				].PlatformFormat	= MTLPixelFormatPVRTC_RGBA_4BPP;
 	GPixelFormats[PF_PVRTC4				].Supported			= true;
-	GPixelFormats[PF_ASTC_4x4			].PlatformFormat	= MTLPixelFormatASTC_4x4_sRGB;
+	GPixelFormats[PF_ASTC_4x4			].PlatformFormat	= MTLPixelFormatASTC_4x4_LDR;
 	GPixelFormats[PF_ASTC_4x4			].Supported			= bCanUseASTC;
-	GPixelFormats[PF_ASTC_6x6			].PlatformFormat	= MTLPixelFormatASTC_6x6_sRGB;
+	GPixelFormats[PF_ASTC_6x6			].PlatformFormat	= MTLPixelFormatASTC_6x6_LDR;
 	GPixelFormats[PF_ASTC_6x6			].Supported			= bCanUseASTC;
-	GPixelFormats[PF_ASTC_8x8			].PlatformFormat	= MTLPixelFormatASTC_8x8_sRGB;
+	GPixelFormats[PF_ASTC_8x8			].PlatformFormat	= MTLPixelFormatASTC_8x8_LDR;
 	GPixelFormats[PF_ASTC_8x8			].Supported			= bCanUseASTC;
-	GPixelFormats[PF_ASTC_10x10			].PlatformFormat	= MTLPixelFormatASTC_10x10_sRGB;
+	GPixelFormats[PF_ASTC_10x10			].PlatformFormat	= MTLPixelFormatASTC_10x10_LDR;
 	GPixelFormats[PF_ASTC_10x10			].Supported			= bCanUseASTC;
-	GPixelFormats[PF_ASTC_12x12			].PlatformFormat	= MTLPixelFormatASTC_12x12_sRGB;
+	GPixelFormats[PF_ASTC_12x12			].PlatformFormat	= MTLPixelFormatASTC_12x12_LDR;
 	GPixelFormats[PF_ASTC_12x12			].Supported			= bCanUseASTC;
 	GPixelFormats[PF_UYVY				].PlatformFormat	= MTLPixelFormatInvalid;
 	GPixelFormats[PF_FloatRGB			].PlatformFormat	= MTLPixelFormatRGBA16Float;
