fix compile issue in no unity build

diff --git a/Engine/Source/Developer/TextureFormatASTC/Private/TextureFormatASTC.cpp b/Engine/Source/Developer/TextureFormatASTC/Private/TextureFormatASTC.cpp
--- a/Engine/Source/Developer/TextureFormatASTC/Private/TextureFormatASTC.cpp
+++ b/Engine/Source/Developer/TextureFormatASTC/Private/TextureFormatASTC.cpp
@@ -20,6 +20,7 @@
 #include "Serialization/CompactBinaryWriter.h"
 #include "TextureBuildFunction.h"
 #include "DerivedDataBuildFunctionFactory.h"
+#include "HAL/IConsoleManager.h"
 
 int32 GASTCCompressor = 0;
 static FAutoConsoleVariableRef CVarASTCCompressor(
