Fix a compiling issue of ASTC compressor.

diff --git a/Engine/Source/Developer/TextureFormatASTC/Private/TextureFormatASTC.cpp b/Engine/Source/Developer/TextureFormatASTC/Private/TextureFormatASTC.cpp
--- a/Engine/Source/Developer/TextureFormatASTC/Private/TextureFormatASTC.cpp
+++ b/Engine/Source/Developer/TextureFormatASTC/Private/TextureFormatASTC.cpp
@@ -236,7 +236,7 @@ static bool CompressSliceToASTC(
 
 	FGuid Guid;
 	FPlatformMisc::CreateGuid(Guid);
-	FString InputFilePath = FString::Printf(bHDR? TEXT("Cache/%08x-%08x-%08x-%08x-RGBToASTCIn.exr"): TEXT("Cache/%08x-%08x-%08x-%08x-RGBToASTCIn.png"), Guid.A, Guid.B, Guid.C, Guid.D);
+	FString InputFilePath = FString::Printf(TEXT("Cache/%08x-%08x-%08x-%08x-RGBToASTCIn"), Guid.A, Guid.B, Guid.C, Guid.D) + (bHDR ? TEXT(".exr") : TEXT(".png"));
 	FString OutputFilePath = FString::Printf(TEXT("Cache/%08x-%08x-%08x-%08x-RGBToASTCOut.astc"), Guid.A, Guid.B, Guid.C, Guid.D);
 
 	InputFilePath  = FPaths::ProjectIntermediateDir() + InputFilePath;
@@ -252,7 +252,7 @@ static bool CompressSliceToASTC(
 	delete PNGFile;
 
 	// Compress PNG file to ASTC (using the reference astcenc.exe from ARM)
-	FString Params = FString::Printf(bHDR ? TEXT("-ch \"%s\" \"%s\" %s") : TEXT("-cl \"%s\" \"%s\" %s"),
+	FString Params = (bHDR ? TEXT("-ch ") : TEXT("-cl ")) + FString::Printf(TEXT("\"%s\" \"%s\" %s"),
 		*InputFilePath,
 		*OutputFilePath,
 		*CompressionParameters
