fix many bad things in TextureFormatASTC
also robustinate

diff --git a/Engine/Source/Developer/TextureFormatASTC/Private/TextureFormatASTC.cpp b/Engine/Source/Developer/TextureFormatASTC/Private/TextureFormatASTC.cpp
--- a/Engine/Source/Developer/TextureFormatASTC/Private/TextureFormatASTC.cpp
+++ b/Engine/Source/Developer/TextureFormatASTC/Private/TextureFormatASTC.cpp
@@ -23,7 +23,15 @@
 #include "DerivedDataSharedString.h"
 #include "HAL/IConsoleManager.h"
 
-// when GASTCCompressor == 0 ,use TextureFormatIntelISPCTexComp instead
+/****
+* 
+* TextureFormatASTC runs the ARM astcenc.exe command line tool
+* 
+* or (by default) redirects to Intel ISPC texcomp* 
+* 
+*****/
+
+// when GASTCCompressor == 0 ,use TextureFormatIntelISPCTexComp instead of this
 int32 GASTCCompressor = 0;
 static FAutoConsoleVariableRef CVarASTCCompressor(
 	TEXT("cook.ASTCTextureCompressor"),
@@ -114,10 +122,6 @@ class FASTCTextureBuildFunction final : public FTextureBuildFunction
 #endif
 
 
-// @@!! this is a global !?
-IImageWrapperModule& ImageWrapperModule = FModuleManager::LoadModuleChecked<IImageWrapperModule>(FName("ImageWrapper"));
-
-
 static int32 GetDefaultCompressionBySizeValue(FCbObjectView InFormatConfigOverride)
 {
 	if (InFormatConfigOverride)
@@ -226,41 +230,53 @@ static uint16 GetQualityVersion(const FCbObjectView& InFormatConfigOverride, int
 }
 
 static bool CompressSliceToASTC(
-	const void* SourceData,
-	int32 SizeX,
-	int32 SizeY,
+	const FImageView & SourceImage,
 	FString CompressionParameters,
 	TArray64<uint8>& OutCompressedData,
-	int32 BitsPerChannel, // @@!! actually BytesPerPixel
-	ERGBFormat Format
-)
+	IImageWrapperModule& ImageWrapperModule
+	)
 {
-	bool bHDR = Format == ERGBFormat::RGBAF;
-	// Compress and retrieve the PNG data to write out to disk
-	TSharedPtr<IImageWrapper> ImageWrapper = ImageWrapperModule.CreateImageWrapper(bHDR ? EImageFormat::EXR : EImageFormat::PNG);
-	// "BitsPerChannel" is actually "BytesPerPixel"
-	// it's either 4 for BGRA8 and 8 for RGBA16f
-	ImageWrapper->SetRaw(SourceData, SizeX * SizeY * BitsPerChannel, SizeX, SizeY, Format, BitsPerChannel * 2);
-	const TArray64<uint8> FileData = ImageWrapper->GetCompressed();
+	bool bHDR = ERawImageFormat::IsHDR(SourceImage.Format);
+	
+	EImageFormat FileFormat = bHDR ? EImageFormat::EXR : EImageFormat::PNG;
+	TArray64<uint8> FileData;
+	bool bCompressSucceeded = ImageWrapperModule.CompressImage(FileData,FileFormat,SourceImage,(int32)EImageCompressionQuality::Uncompressed);
+	if ( ! bCompressSucceeded )
+	{
+		UE_LOG(LogTextureFormatASTC, Error, TEXT("CompressSliceToASTC CompressImage failed"));
+		return false;
+	}
+
 	int64 FileDataSize = FileData.Num();
 
+	// make a random file name to write the image :
 	FGuid Guid;
 	FPlatformMisc::CreateGuid(Guid);
-	FString InputFilePath = FString::Printf(TEXT("Cache/%08x-%08x-%08x-%08x-RGBToASTCIn"), Guid.A, Guid.B, Guid.C, Guid.D) + (bHDR ? TEXT(".exr") : TEXT(".png"));
+	FString InputFilePath = FString::Printf(TEXT("Cache/%08x-%08x-%08x-%08x-RGBToASTCIn"), Guid.A, Guid.B, Guid.C, Guid.D) + TEXT(".") + ImageWrapperModule.GetExtension(FileFormat);
 	FString OutputFilePath = FString::Printf(TEXT("Cache/%08x-%08x-%08x-%08x-RGBToASTCOut.astc"), Guid.A, Guid.B, Guid.C, Guid.D);
 
 	InputFilePath  = FPaths::ProjectIntermediateDir() + InputFilePath;
 	OutputFilePath = FPaths::ProjectIntermediateDir() + OutputFilePath;
 
-	FArchive* PNGFile = NULL;
-	while (!PNGFile)
+	// write to InputFilePath :
 	{
-		PNGFile = IFileManager::Get().CreateFileWriter(*InputFilePath);   // Occasionally returns NULL due to error code ERROR_SHARING_VIOLATION
-		// @@!! don't sleep if not null !
-		FPlatformProcess::Sleep(0.01f);                             // ... no choice but to wait for the file to become free to access
+		FArchive* PNGFile = IFileManager::Get().CreateFileWriter(*InputFilePath);
+		while (!PNGFile)
+		{
+			// CreateFileWriter occasionally returns NULL due to error code ERROR_SHARING_VIOLATION
+			// ... no choice but to wait for the file to become free to access
+		
+			UE_LOG(LogTextureFormatASTC, Display, TEXT("CreateFileWriter for %s failed, trying again..."), *InputFilePath);
+
+			FPlatformProcess::Sleep(0.01f);                             
+			PNGFile = IFileManager::Get().CreateFileWriter(*InputFilePath);   
+		}
+		PNGFile->Serialize((void*)&FileData[0], FileDataSize);
+		delete PNGFile;
 	}
-	PNGFile->Serialize((void*)&FileData[0], FileDataSize);
-	delete PNGFile;
+
+	// FileData written, can free now :
+	FileData.Reset();
 
 	// Compress PNG file to ASTC (using the reference astcenc.exe from ARM)
 	FString Params = (bHDR ? TEXT("-ch ") : TEXT("-cl ")) + FString::Printf(TEXT("\"%s\" \"%s\" %s"),
@@ -283,33 +299,42 @@ static bool CompressSliceToASTC(
 #else
 #error Unsupported platform
 #endif
-	FProcHandle Proc = FPlatformProcess::CreateProc(*CompressorPath, *Params, true, false, false, NULL, -1, NULL, NULL);
 
-	// Failed to start the compressor process
-	if (!Proc.IsValid())
+	// run the astcenc process :
 	{
-		UE_LOG(LogTextureFormatASTC, Error, TEXT("Failed to start astcenc for compressing images (%s)"), *CompressorPath);
-		return false;
-	}
+		FProcHandle Proc = FPlatformProcess::CreateProc(*CompressorPath, *Params, true, false, false, NULL, -1, NULL, NULL);
 
-	// Wait for the process to complete
-	// @@!! WaitForProc instead!
-	int ReturnCode;
-	while (!FPlatformProcess::GetProcReturnCode(Proc, &ReturnCode))
-	{
-		FPlatformProcess::Sleep(0.01f);
-	}
+		// Failed to start the compressor process
+		if (!Proc.IsValid())
+		{
+			UE_LOG(LogTextureFormatASTC, Error, TEXT("Failed to start astcenc for compressing images (%s)"), *CompressorPath);
+			return false;
+		}
 
-	// Did it work?
-	bool bConversionWasSuccessful = (ReturnCode == 0);
+		// Wait for the process to complete
+		FPlatformProcess::WaitForProc(Proc);
+		int ReturnCode = -1;
+		FPlatformProcess::GetProcReturnCode(Proc, &ReturnCode);
+		FPlatformProcess::CloseProc(Proc);
+		
+		// Did it work?
+		if ( ReturnCode != 0)
+		{
+			UE_LOG(LogTextureFormatASTC, Error, TEXT("ASTC encoder failed with return code %d, mip size (%d, %d). Leaving '%s' for testing."), ReturnCode, SourceImage.SizeX, SourceImage.SizeY, *InputFilePath);
+			return false;
+		}
+	}
 
 	// Open compressed file and put the data in OutCompressedImage
-	if (bConversionWasSuccessful)
 	{
 		// Get raw file data
 		TArray64<uint8> ASTCData;
-		FFileHelper::LoadFileToArray(ASTCData, *OutputFilePath);
-			
+		if ( ! FFileHelper::LoadFileToArray(ASTCData, *OutputFilePath) )
+		{
+			UE_LOG(LogTextureFormatASTC, Error, TEXT("Failed load output of astcenc (%s -> %s)"),*InputFilePath,*OutputFilePath);
+			return false;
+		}
+
 		// Process it
 		FASTCHeader* Header = (FASTCHeader*)ASTCData.GetData();
 			
@@ -327,19 +352,19 @@ static bool CompressSliceToASTC(
 			(Header->TexelCountZ[1] <<  8) + 
 			(Header->TexelCountZ[2] << 16);
 
-//		UE_LOG(LogTextureFormatASTC, Display, TEXT("    Compressed Texture Header:"));
-//		UE_LOG(LogTextureFormatASTC, Display, TEXT("             Magic: %x"), Header->Magic);
-//		UE_LOG(LogTextureFormatASTC, Display, TEXT("        BlockSizeX: %u"), Header->BlockSizeX);
-//		UE_LOG(LogTextureFormatASTC, Display, TEXT("        BlockSizeY: %u"), Header->BlockSizeY);
-//		UE_LOG(LogTextureFormatASTC, Display, TEXT("        BlockSizeZ: %u"), Header->BlockSizeZ);
-//		UE_LOG(LogTextureFormatASTC, Display, TEXT("       TexelCountX: %u"), TexelCountX);
-//		UE_LOG(LogTextureFormatASTC, Display, TEXT("       TexelCountY: %u"), TexelCountY);
-//		UE_LOG(LogTextureFormatASTC, Display, TEXT("       TexelCountZ: %u"), TexelCountZ);
+		if ( TexelCountX != SourceImage.SizeX ||
+			 TexelCountY != SourceImage.SizeY )
+		{
+			UE_LOG(LogTextureFormatASTC, Warning, TEXT("Unexpected image size mismatch : %d x %d != %d x %d"),
+				TexelCountX,TexelCountY,SourceImage.SizeX,SourceImage.SizeY);
+		}
 
 		// Calculate size of this mip in blocks
 		uint32 MipSizeX = (TexelCountX + Header->BlockSizeX - 1) / Header->BlockSizeX;
 		uint32 MipSizeY = (TexelCountY + Header->BlockSizeY - 1) / Header->BlockSizeY;
 
+		// TexelCountZ ignored
+
 		// A block is always 16 bytes
 		uint64 MipSize = (uint64)MipSizeX * MipSizeY * 16;
 
@@ -353,11 +378,6 @@ static bool CompressSliceToASTC(
 		check(ASTCData.Num() == (sizeof(FASTCHeader) + MipSize));
 		FMemory::Memcpy(MipData, ASTCData.GetData() + sizeof(FASTCHeader), MipSize);
 	}
-	else
-	{
-		UE_LOG(LogTextureFormatASTC, Error, TEXT("ASTC encoder failed with return code %d, mip size (%d, %d). Leaving '%s' for testing."), ReturnCode, SizeX, SizeY, *InputFilePath);
-		return false;
-	}
 		
 	// Delete intermediate files
 	IFileManager::Get().Delete(*InputFilePath);
@@ -372,8 +392,11 @@ class FTextureFormatASTC : public ITextureFormat
 {
 public:
 	FTextureFormatASTC()
-	:	IntelISPCTexCompFormat(*FModuleManager::LoadModuleChecked<ITextureFormatModule>(TEXT("TextureFormatIntelISPCTexComp")).GetTextureFormat())
+	:	IntelISPCTexCompFormat(*FModuleManager::LoadModuleChecked<ITextureFormatModule>(TEXT("TextureFormatIntelISPCTexComp")).GetTextureFormat()),
+		ImageWrapperModule(FModuleManager::LoadModuleChecked<IImageWrapperModule>(FName("ImageWrapper")))
 	{
+		// LoadModule has to be done on Main thread
+		// can't be done on-demand in the Compress call
 	}
 
 	virtual bool AllowParallelBuild() const override
@@ -469,13 +492,23 @@ public:
 			return IntelISPCTexCompFormat.CompressImage(InImage, BuildSettings, DebugTexturePathName, bImageHasAlphaChannel, OutCompressedImage);
 		}
 #endif
+
+		// note: BuildSettings.bHDRSource is a weird thing to be checking here
+		//  it should look at whether the requested output format is HDR
+		//  the input->output mapping has already been done, you shouldn't do it again
 		bool bHDRImage = BuildSettings.bHDRSource && GASTCHDRProfile;
-		// Get Raw Image Data from passed in FImage
+		if ( BuildSettings.bHDRSource && ! GASTCHDRProfile )
+		{
+			UE_LOG(LogTextureFormatASTC, Display, TEXT("HDR Image passed to ASTC but ASTC HDR is disabled, will compress as 8 bit: (%s)"),*FString(DebugTexturePathName));
+		}
+
+		// Get Raw Image Data from passed in FImage & convert to BGRA8 or RGBA16F
+		// note: wasteful, often copies image to same format
 		FImage Image;
 		InImage.CopyTo(Image, bHDRImage ? ERawImageFormat::RGBA16F : ERawImageFormat::BGRA8, BuildSettings.GetDestGammaSpace());
 
 		// Determine the compressed pixel format and compression parameters
-		EPixelFormat CompressedPixelFormat = GetPixelFormatForImage(BuildSettings, InImage, bImageHasAlphaChannel);
+		EPixelFormat CompressedPixelFormat = GetPixelFormatForImage(BuildSettings, Image, bImageHasAlphaChannel);
 
 		FString CompressionParameters = TEXT("");
 
@@ -507,38 +540,22 @@ public:
 		{
 			CompressionParameters = FString::Printf(TEXT("%s -esw bg00 -cw 1 1 0 0 -dblimit 60 -b 2.5 -v 3 1 1 0 50 0 -va 1 1 0 50"), *GetQualityString(BuildSettings.FormatConfigOverride, FORCED_NORMAL_MAP_COMPRESSION_SIZE_VALUE, -1));
 		}
+		else
+		{
+			check(false);
+		}
 
 		// Compress the image, slice by slice
-		// @@!! use FImage.GetSlice()
 		bool bCompressionSucceeded = true;
-		int32 SliceSizeInTexels = Image.SizeX * Image.SizeY;
+
 		for (int32 SliceIndex = 0; SliceIndex < Image.NumSlices; ++SliceIndex)
 		{
 			TArray64<uint8> CompressedSliceData;
-			if (bHDRImage)
-			{
-				bCompressionSucceeded = CompressSliceToASTC(
-					(&Image.AsRGBA16F()[0]) + (SliceIndex * SliceSizeInTexels),
-					Image.SizeX,
-					Image.SizeY,
-					CompressionParameters,
-					CompressedSliceData,
-					8,
-					ERGBFormat::RGBAF
-				);
-			}
-			else
-			{
-				bCompressionSucceeded = CompressSliceToASTC(
-					(&Image.AsBGRA8()[0]) + (SliceIndex * SliceSizeInTexels),
-					Image.SizeX,
-					Image.SizeY,
-					CompressionParameters,
-					CompressedSliceData,
-					4,
-					ERGBFormat::RGBA
-				);
-			}
+
+			FImageView Slice = Image.GetSlice(SliceIndex);
+			
+			bCompressionSucceeded = CompressSliceToASTC(Slice,CompressionParameters,CompressedSliceData,ImageWrapperModule);
+
 			if ( ! bCompressionSucceeded )
 			{
 				return false;
@@ -558,6 +575,8 @@ public:
 
 private:
 	const ITextureFormat& IntelISPCTexCompFormat;
+	
+	IImageWrapperModule& ImageWrapperModule;
 };
 
 /**
