Fix ASTC mapping for BC6 and BC7
Fix ASTC HDR support
Fix up ASTC & ISPC texture formats a bit

diff --git a/Engine/Source/Developer/Android/AndroidTargetPlatform/Private/AndroidTargetPlatform.cpp b/Engine/Source/Developer/Android/AndroidTargetPlatform/Private/AndroidTargetPlatform.cpp
--- a/Engine/Source/Developer/Android/AndroidTargetPlatform/Private/AndroidTargetPlatform.cpp
+++ b/Engine/Source/Developer/Android/AndroidTargetPlatform/Private/AndroidTargetPlatform.cpp
@@ -54,7 +54,6 @@ namespace AndroidTexFormat
 {
 	// Compressed Texture Formats
 	const static FName NameDXT1(TEXT("DXT1"));
-	const static FName NameDXT3(TEXT("DXT3"));
 	const static FName NameDXT5(TEXT("DXT5"));
 	const static FName NameDXT5n(TEXT("DXT5n"));
 	const static FName NameAutoDXT(TEXT("AutoDXT"));
@@ -62,25 +61,19 @@ namespace AndroidTexFormat
 	const static FName NameBC5(TEXT("BC5"));
 	const static FName NameBC6H(TEXT("BC6H"));
 	const static FName NameBC7(TEXT("BC7"));
+
 	const static FName NameETC2_RGB(TEXT("ETC2_RGB"));
 	const static FName NameETC2_RGBA(TEXT("ETC2_RGBA"));
 	const static FName NameETC2_R11(TEXT("ETC2_R11"));
 	const static FName NameAutoETC2(TEXT("AutoETC2"));
-	// AFAICT none of these explicitly sized ASTC names are actually valid TextureFormat names
-	const static FName NameASTC_4x4(TEXT("ASTC_4x4"));
-	const static FName NameASTC_6x6(TEXT("ASTC_6x6"));
-	const static FName NameASTC_8x8(TEXT("ASTC_8x8"));
-	const static FName NameASTC_10x10(TEXT("ASTC_10x10"));
-	const static FName NameASTC_12x12(TEXT("ASTC_12x12"));
+
 	const static FName NameAutoASTC(TEXT("ASTC_RGBAuto"));
 	
 	// Uncompressed Texture Formats
 	const static FName NameBGRA8(TEXT("BGRA8"));
 	const static FName NameG8(TEXT("G8"));
-	const static FName NameVU8(TEXT("VU8"));
 	const static FName NameRGBA16F(TEXT("RGBA16F"));
 	const static FName NameR16F(TEXT("R16F"));
-	const static FName NameR5G6B5(TEXT("R5G6B5"));
 
 	//A1RGB555 is mapped to RGB555A1, because OpenGL GL_RGB5_A1 only supports alpha on the lowest bit
 	const static FName NameA1RGB555(TEXT("A1RGB555"));
@@ -88,9 +81,10 @@ namespace AndroidTexFormat
 
 	const static FName GenericRemap[][2] =
 	{
-		{ NameBC6H,			NameRGBA16F					},
-		{ NameBC7,			NameBGRA8					},
+		{ NameA1RGB555,		NameRGB555A1				},
 	};
+	
+	static const FName NameASTC_RGB_HDR(TEXT("ASTC_RGB_HDR"));
 
 	const static FName ASTCRemap[][2] =
 	{
@@ -100,12 +94,9 @@ namespace AndroidTexFormat
 		{ NameDXT5n,		FName(TEXT("ASTC_NormalAG"))},
 		{ NameBC5,			FName(TEXT("ASTC_NormalRG"))},
 		{ NameBC4,			NameETC2_R11				},
-		{ NameBC6H,			NameRGBA16F					},
-		{ NameBC7,			NameBGRA8					},
-		//{ NameBC6H,			FName(TEXT("ASTC_RGB"))		},
-		//{ NameBC7,			NameAutoASTC				},
+		{ NameBC6H,			NameASTC_RGB_HDR			},
+		{ NameBC7,			FName(TEXT("ASTC_RGBA_HQ"))	},
 		{ NameAutoDXT,		NameAutoASTC				},
-		{ NameA1RGB555,		NameRGB555A1				}
 	};
 
 	const static FName ETCRemap[][2] =
@@ -116,12 +107,9 @@ namespace AndroidTexFormat
 		{ NameDXT5n,		NameETC2_RGB	},
 		{ NameBC5,			NameETC2_RGB	},
 		{ NameBC4,			NameETC2_R11	},
-		{ NameBC6H,			NameRGBA16F					},
-		{ NameBC7,			NameBGRA8					},
-		//{ NameBC6H,			NameETC2_RGB	},
-		//{ NameBC7,			NameAutoETC2	},
+		{ NameBC6H,			NameRGBA16F		},
+		{ NameBC7,			NameETC2_RGBA	},
 		{ NameAutoDXT,		NameAutoETC2	},
-		{ NameA1RGB555,		NameRGB555A1	}
 	};
 }
 
@@ -493,6 +481,10 @@ const FStaticMeshLODSettings& FAndroidTargetPlatform::GetStaticMeshLODSettings(
 
 void FAndroidTargetPlatform::GetTextureFormats( const UTexture* Texture, TArray< TArray<FName> >& OutFormats) const
 {
+	// FAndroidTargetPlatform aside from being the base class for all the concrete android target platforms
+	//	it is also usable on its own as "flavorless" Android
+	//	but I don't understand how that's supposed to work or what that's supposed to mean
+	//	and no information has been forthcoming
 	check(Texture);
 	
 	// Supported in ES3.2 with ASTC
@@ -501,7 +493,8 @@ void FAndroidTargetPlatform::GetTextureFormats( const UTexture* Texture, TArray<
 	const bool bSupportFilteredFloat32Textures = false;
 
 	TArray<FName>& LayerFormats = OutFormats.AddDefaulted_GetRef();
-	GetDefaultTextureFormatNamePerLayer(LayerFormats, this, Texture, bSupportCompressedVolumeTexture, 1, bSupportFilteredFloat32Textures);
+	int32 BlockSize = 1; // this looks wrong? should be 4 for FAndroid_DXTTargetPlatform ? - it is wrong, but BlockSize is ignored
+	GetDefaultTextureFormatNamePerLayer(LayerFormats, this, Texture, bSupportCompressedVolumeTexture, BlockSize, bSupportFilteredFloat32Textures);
 
 	for (FName& TextureFormatName : LayerFormats)
 	{
@@ -510,11 +503,11 @@ void FAndroidTargetPlatform::GetTextureFormats( const UTexture* Texture, TArray<
 			// forward rendering only needs one channel for shadow maps
 			TextureFormatName = FName(TEXT("G8"));
 		}
-
+		
 		if (Texture->IsA(UTextureCube::StaticClass()))
 		{
 			const UTextureCube* Cube = CastChecked<UTextureCube>(Texture);
-			if (Cube != nullptr)
+			if (Cube != nullptr) 
 			{
 				FTextureFormatSettings FormatSettings;
 				Cube->GetDefaultFormatSettings(FormatSettings);
@@ -545,15 +538,19 @@ FName FAndroidTargetPlatform::FinalizeVirtualTextureLayerFormat(FName Format) co
 	// currently it forces all ASTC to ETC
 	// this is needed because the runtime virtual texture encoder only supports ETC
 
+	// code dupe with IOSTargetPlatform
+
 	const static FName VTRemap[][2] =
 	{
 		{ { FName(TEXT("ASTC_RGB")) },			{ AndroidTexFormat::NameETC2_RGB } },
 		{ { FName(TEXT("ASTC_RGBA")) },			{ AndroidTexFormat::NameETC2_RGBA } },
+		{ { FName(TEXT("ASTC_RGBA_HQ")) },		{ AndroidTexFormat::NameETC2_RGBA } },
+//		{ { FName(TEXT("ASTC_RGB_HDR")) },		{ NameRGBA16F } }, // ?
 		{ { FName(TEXT("ASTC_RGBAuto")) },		{ AndroidTexFormat::NameAutoETC2 } },
 		{ { FName(TEXT("ASTC_NormalAG")) },		{ AndroidTexFormat::NameETC2_RGB } },
 		{ { FName(TEXT("ASTC_NormalRG")) },		{ AndroidTexFormat::NameETC2_RGB } },
 		{ { AndroidTexFormat::NameDXT1 },		{ AndroidTexFormat::NameETC2_RGB } },
-		{ { AndroidTexFormat::NameDXT5 },		{ AndroidTexFormat::NameAutoETC2 } },
+		{ { AndroidTexFormat::NameDXT5 },		{ AndroidTexFormat::NameETC2_RGBA } },
 		{ { AndroidTexFormat::NameAutoDXT },	{ AndroidTexFormat::NameAutoETC2 } }
 	};
 
@@ -585,20 +582,8 @@ void FAndroidTargetPlatform::GetAllTextureFormats(TArray<FName>& OutFormats) con
 
 void FAndroid_ASTCTargetPlatform::GetAllTextureFormats(TArray<FName>& OutFormats) const
 {
-	GetAllDefaultTextureFormats(this, OutFormats);
-
-	for (int32 RemapIndex = 0; RemapIndex < UE_ARRAY_COUNT(AndroidTexFormat::GenericRemap); ++RemapIndex)
-	{
-		OutFormats.Remove(AndroidTexFormat::GenericRemap[RemapIndex][0]);
-	}
+	FAndroidTargetPlatform::GetAllTextureFormats(OutFormats);
 
-	for (int32 RemapIndex = 0; RemapIndex < UE_ARRAY_COUNT(AndroidTexFormat::GenericRemap); ++RemapIndex)
-	{
-		OutFormats.AddUnique(AndroidTexFormat::GenericRemap[RemapIndex][1]);
-	}
-
-	// not supported
-	OutFormats.Remove(AndroidTexFormat::NameDXT3);
 	for (int32 RemapIndex = 0; RemapIndex < UE_ARRAY_COUNT(AndroidTexFormat::ASTCRemap); ++RemapIndex)
 	{
 		OutFormats.Remove(AndroidTexFormat::ASTCRemap[RemapIndex][0]);
@@ -627,24 +612,49 @@ void FAndroid_ASTCTargetPlatform::GetTextureFormats(const UTexture* Texture, TAr
 			}
 		}
 	}
+	
+	bool bSupportASTCHDR = UsesASTCHDR();
+
+	if ( ! bSupportASTCHDR )
+	{
+		for (FName& TextureFormatName : LayerFormats)
+		{
+			if ( TextureFormatName == AndroidTexFormat::NameASTC_RGB_HDR )
+			{
+				TextureFormatName = AndroidTexFormat::NameRGBA16F;
+			}
+		}
+	}
 }
 
-void FAndroid_ETC2TargetPlatform::GetAllTextureFormats(TArray<FName>& OutFormats) const
+void FAndroid_DXTTargetPlatform::GetTextureFormats(const UTexture* Texture, TArray< TArray<FName> >& OutFormats) const
 {
-	GetAllDefaultTextureFormats(this, OutFormats);
+	FAndroidTargetPlatform::GetTextureFormats(Texture, OutFormats);
 
-	for (int32 RemapIndex = 0; RemapIndex < UE_ARRAY_COUNT(AndroidTexFormat::GenericRemap); ++RemapIndex)
-	{
-		OutFormats.Remove(AndroidTexFormat::GenericRemap[RemapIndex][0]);
-	}
+	bool bSupportsDX11Formats = false; // assume Android DXT does not support BC6/7
 
-	for (int32 RemapIndex = 0; RemapIndex < UE_ARRAY_COUNT(AndroidTexFormat::GenericRemap); ++RemapIndex)
+	if ( ! bSupportsDX11Formats )
 	{
-		OutFormats.AddUnique(AndroidTexFormat::GenericRemap[RemapIndex][1]);
+		TArray<FName>& LayerFormats = OutFormats.Last();
+
+		for (FName& TextureFormatName : LayerFormats)
+		{
+			if ( TextureFormatName == AndroidTexFormat::NameBC6H )
+			{
+				TextureFormatName = AndroidTexFormat::NameRGBA16F;
+			}
+			else if ( TextureFormatName == AndroidTexFormat::NameBC7 )
+			{
+				TextureFormatName = AndroidTexFormat::NameDXT5;
+			}
+		}
 	}
+}
+
+void FAndroid_ETC2TargetPlatform::GetAllTextureFormats(TArray<FName>& OutFormats) const
+{
+	FAndroidTargetPlatform::GetAllTextureFormats(OutFormats);
 
-	// not supported
-	OutFormats.Remove(AndroidTexFormat::NameDXT3);
 	for (int32 RemapIndex = 0; RemapIndex < UE_ARRAY_COUNT(AndroidTexFormat::ETCRemap); ++RemapIndex)
 	{
 		OutFormats.Remove(AndroidTexFormat::ETCRemap[RemapIndex][0]);
diff --git a/Engine/Source/Developer/Android/AndroidTargetPlatform/Private/AndroidTargetPlatform.h b/Engine/Source/Developer/Android/AndroidTargetPlatform/Private/AndroidTargetPlatform.h
--- a/Engine/Source/Developer/Android/AndroidTargetPlatform/Private/AndroidTargetPlatform.h
+++ b/Engine/Source/Developer/Android/AndroidTargetPlatform/Private/AndroidTargetPlatform.h
@@ -230,6 +230,10 @@ public:
 	FAndroid_DXTTargetPlatform(bool bIsClient) : FAndroidTargetPlatform(bIsClient, TEXT("DXT"))
 	{
 	}
+	
+#if WITH_ENGINE
+	virtual void GetTextureFormats(const UTexture* Texture, TArray< TArray<FName> >& OutFormats) const override;
+#endif
 
 	virtual bool SupportsTextureFormatCategory(EAndroidTextureFormatCategory Category) const override
 	{
diff --git a/Engine/Source/Developer/Android/AndroidTargetPlatform/Private/AndroidTargetPlatformModule.cpp b/Engine/Source/Developer/Android/AndroidTargetPlatform/Private/AndroidTargetPlatformModule.cpp
--- a/Engine/Source/Developer/Android/AndroidTargetPlatform/Private/AndroidTargetPlatformModule.cpp
+++ b/Engine/Source/Developer/Android/AndroidTargetPlatform/Private/AndroidTargetPlatformModule.cpp
@@ -24,10 +24,16 @@ public:
 	{
 		if (FAndroidTargetPlatform::IsUsable())
 		{
+			 // ensure this is only called once
+			check( SinglePlatforms.IsEmpty() );
+
 			for (int32 Type = 0; Type < 2; Type++)
 			{
 				bool bIsClient = Type == 1;
+
 				// flavorless must come first
+				// "flavorless" (aka "vanilla") Android is not well defined and not actually usable
+				//  but Editor wants to see it in the platform list
 				SinglePlatforms.Add(new FAndroidTargetPlatform(bIsClient, nullptr));
 				SinglePlatforms.Add(new FAndroid_ASTCTargetPlatform(bIsClient));
 				SinglePlatforms.Add(new FAndroid_DXTTargetPlatform(bIsClient));
diff --git a/Engine/Source/Developer/HoloLens/HoloLensTargetPlatform/Private/HoloLensTargetPlatform.cpp b/Engine/Source/Developer/HoloLens/HoloLensTargetPlatform/Private/HoloLensTargetPlatform.cpp
--- a/Engine/Source/Developer/HoloLens/HoloLensTargetPlatform/Private/HoloLensTargetPlatform.cpp
+++ b/Engine/Source/Developer/HoloLens/HoloLensTargetPlatform/Private/HoloLensTargetPlatform.cpp
@@ -191,17 +191,18 @@ void FHoloLensTargetPlatform::GetReflectionCaptureFormats(TArray<FName>& OutForm
 const static FName NameBC6H(TEXT("BC6H"));
 const static FName NameBC7(TEXT("BC7"));
 const static FName NameRGBA16F(TEXT("RGBA16F"));
-const static FName NameBGRA8(TEXT("BGRA8"));
+const static FName NameDXT5(TEXT("DXT5"));
 const static FName HoloLensRemap[][2] =
 {
 	{ NameBC6H,			NameRGBA16F		},
-	{ NameBC7,			NameBGRA8		},
+	{ NameBC7,			NameDXT5		},
 };
 
 void FHoloLensTargetPlatform::GetTextureFormats(const UTexture* InTexture, TArray< TArray<FName> >& OutFormats) const
 {
-	bool bSupportFilteredFloat32Textures = false;
-	GetDefaultTextureFormatNamePerLayer(OutFormats.AddDefaulted_GetRef(), this, InTexture, false, 4, bSupportFilteredFloat32Textures);
+	const bool bSupportFilteredFloat32Textures = false;
+	const bool bSupportCompressedVolumeTexture = false;
+	GetDefaultTextureFormatNamePerLayer(OutFormats.AddDefaulted_GetRef(), this, InTexture, bSupportCompressedVolumeTexture, 4, bSupportFilteredFloat32Textures);
 
 	// perform any remapping away from defaults
 	TArray<FName>& LayerFormats = OutFormats.Last();
diff --git a/Engine/Source/Developer/IOS/IOSTargetPlatform/Private/IOSTargetPlatform.cpp b/Engine/Source/Developer/IOS/IOSTargetPlatform/Private/IOSTargetPlatform.cpp
--- a/Engine/Source/Developer/IOS/IOSTargetPlatform/Private/IOSTargetPlatform.cpp
+++ b/Engine/Source/Developer/IOS/IOSTargetPlatform/Private/IOSTargetPlatform.cpp
@@ -570,9 +570,10 @@ void FIOSTargetPlatform::GetReflectionCaptureFormats( TArray<FName>& OutFormats
 	OutFormats.Add(FName(TEXT("EncodedHDR")));
 }
 
+static const FName NameASTC_RGB_HDR(TEXT("ASTC_RGB_HDR"));
 
 // we remap some of the defaults
-static FName FormatRemap[] =
+static const FName FormatRemap[] =
 {
 	// original				ASTC
 	FName(TEXT("AutoDXT")),	FName(TEXT("ASTC_RGBAuto")),
@@ -581,11 +582,11 @@ static FName FormatRemap[] =
 	FName(TEXT("DXT5n")),	FName(TEXT("ASTC_NormalAG")),
 	FName(TEXT("BC5")),		FName(TEXT("ASTC_NormalRG")),
 	FName(TEXT("BC4")),		FName(TEXT("ETC2_R11")),
-	FName(TEXT("BC6H")),	FName(TEXT("ASTC_RGB")), 
-	FName(TEXT("BC7")),		FName(TEXT("ASTC_RGBAuto"))
+	FName(TEXT("BC6H")),	NameASTC_RGB_HDR,
+	FName(TEXT("BC7")),		FName(TEXT("ASTC_RGBA_HQ"))
 };
-static FName NameBGRA8(TEXT("BGRA8"));
-static FName NameG8 = FName(TEXT("G8"));
+static const FName NameG8(TEXT("G8"));
+static const FName NameRGBA16F(TEXT("RGBA16F"));
 
 void FIOSTargetPlatform::GetTextureFormats( const UTexture* Texture, TArray< TArray<FName> >& OutFormats) const
 {
@@ -609,18 +610,13 @@ void FIOSTargetPlatform::GetTextureFormats( const UTexture* Texture, TArray< TAr
 		// Compressed volume textures require MTLGPUFamilyApple3 or later
 		// min spec for TVOS is AppleTV HD which is MTLGPUFamilyApple2 (A8)
 		bool bSupportCompressedVolumeTexture = !bIsTVOS && !SupportsA8Devices();
-		GetDefaultTextureFormatNamePerLayer(TextureFormatNames, this, Texture, bSupportCompressedVolumeTexture, BlockSize);
+		bool bSupportFilteredFloat32Textures = false;
+		GetDefaultTextureFormatNamePerLayer(TextureFormatNames, this, Texture, bSupportCompressedVolumeTexture, BlockSize, bSupportFilteredFloat32Textures);
 	}
 
 	// include the formats we want
 	for (FName& TextureFormatName : TextureFormatNames)
 	{
-		if (TextureFormatName == NameBGRA8 || 
-			TextureFormatName == NameG8)
-		{
-			continue;
-		}
-
 		for (int32 RemapIndex = 0; RemapIndex < UE_ARRAY_COUNT(FormatRemap); RemapIndex += 2)
 		{
 			if (TextureFormatName == FormatRemap[RemapIndex])
@@ -630,7 +626,20 @@ void FIOSTargetPlatform::GetTextureFormats( const UTexture* Texture, TArray< TAr
 			}
 		}
 	}
-		
+	
+	bool bSupportASTCHDR = UsesASTCHDR();
+
+	if ( ! bSupportASTCHDR )
+	{
+		for (FName& TextureFormatName : TextureFormatNames)
+		{
+			if ( TextureFormatName == NameASTC_RGB_HDR )
+			{
+				TextureFormatName = NameRGBA16F;
+			}
+		}
+	}
+
 	for (FName& TextureFormatName : OutFormats.Last())
 	{
 		if (Texture->IsA(UTextureCube::StaticClass()))
@@ -675,6 +684,8 @@ FName FIOSTargetPlatform::FinalizeVirtualTextureLayerFormat(FName Format) const
 	// this then runs again
 	// currently it forces all ASTC to ETC
 	// this is needed because the runtime virtual texture encoder only supports ETC
+	
+	// code dupe with AndroidTargetPlatform
 
 	const static FName NameETC2_RGB(TEXT("ETC2_RGB"));
 	const static FName NameETC2_RGBA(TEXT("ETC2_RGBA"));
@@ -686,6 +697,8 @@ FName FIOSTargetPlatform::FinalizeVirtualTextureLayerFormat(FName Format) const
 		{ { FName(TEXT("ASTC_RGB")) },			{ NameETC2_RGB } },
 		{ { FName(TEXT("ASTC_RGBA")) },			{ NameETC2_RGBA } },
 		{ { FName(TEXT("ASTC_RGBAuto")) },		{ NameAutoETC2 } },
+		{ { FName(TEXT("ASTC_RGBA_HQ")) },		{ NameETC2_RGBA } },
+//		{ { FName(TEXT("ASTC_RGB_HDR")) },		{ NameRGBA16F } }, // ?
 		{ { FName(TEXT("ASTC_NormalAG")) },		{ NameETC2_RGB } },
 		{ { FName(TEXT("ASTC_NormalRG")) },		{ NameETC2_RGB } },
 	};
diff --git a/Engine/Source/Developer/TargetPlatform/Private/TargetPlatformBase.cpp b/Engine/Source/Developer/TargetPlatform/Private/TargetPlatformBase.cpp
--- a/Engine/Source/Developer/TargetPlatform/Private/TargetPlatformBase.cpp
+++ b/Engine/Source/Developer/TargetPlatform/Private/TargetPlatformBase.cpp
@@ -82,6 +82,19 @@ bool FTargetPlatformBase::UsesMobileAmbientOcclusion() const
 	return CVar ? (CVar->GetInt() != 0) : false;
 }
 
+int32 GASTCHDRProfile = 0;
+static FAutoConsoleVariableRef CVarAllowASTCHDRProfile(
+	TEXT("cook.AllowASTCHDRProfile"),
+	GASTCHDRProfile,
+	TEXT("whether to allow ASTC HDR profile, the hdr format is only supported on some devices, e.g. Apple A13, Mali-G72, Adreno (TM) 660"),
+	ECVF_Default | ECVF_ReadOnly
+);
+
+bool FTargetPlatformBase::UsesASTCHDR() const
+{
+	return GASTCHDRProfile != 0;
+}
+
 static bool IsPluginEnabledForTarget(const IPlugin& Plugin, const FProjectDescriptor* Project, const FString& Platform, EBuildConfiguration Configuration, EBuildTargetType TargetType)
 {
 	if (!Plugin.GetDescriptor().SupportsTargetPlatform(Platform))
diff --git a/Engine/Source/Developer/TargetPlatform/Public/Common/TargetPlatformBase.h b/Engine/Source/Developer/TargetPlatform/Public/Common/TargetPlatformBase.h
--- a/Engine/Source/Developer/TargetPlatform/Public/Common/TargetPlatformBase.h
+++ b/Engine/Source/Developer/TargetPlatform/Public/Common/TargetPlatformBase.h
@@ -77,6 +77,8 @@ public:
 
 	TARGETPLATFORM_API virtual bool UsesMobileAmbientOcclusion() const override;
 
+	TARGETPLATFORM_API virtual bool UsesASTCHDR() const override;
+
 #if WITH_ENGINE
 	virtual void GetReflectionCaptureFormats( TArray<FName>& OutFormats ) const override
 	{
diff --git a/Engine/Source/Developer/TargetPlatform/Public/Interfaces/ITargetPlatform.h b/Engine/Source/Developer/TargetPlatform/Public/Interfaces/ITargetPlatform.h
--- a/Engine/Source/Developer/TargetPlatform/Public/Interfaces/ITargetPlatform.h
+++ b/Engine/Source/Developer/TargetPlatform/Public/Interfaces/ITargetPlatform.h
@@ -511,6 +511,11 @@ public:
 	 * Gets whether the platform uses Mobile AO
 	 */
 	virtual bool UsesMobileAmbientOcclusion() const = 0;
+	
+	/**
+	 * Gets whether the platform uses ASTC HDR
+	 */
+	virtual bool UsesASTCHDR() const = 0;
 
 	/**
 	* Gets the shader formats this platform can use.
diff --git a/Engine/Source/Developer/TextureFormatASTC/Private/TextureFormatASTC.cpp b/Engine/Source/Developer/TextureFormatASTC/Private/TextureFormatASTC.cpp
--- a/Engine/Source/Developer/TextureFormatASTC/Private/TextureFormatASTC.cpp
+++ b/Engine/Source/Developer/TextureFormatASTC/Private/TextureFormatASTC.cpp
@@ -32,6 +32,7 @@
 *****/
 
 // when GASTCCompressor == 0 ,use TextureFormatIntelISPCTexComp instead of this
+// @@!! does this break DDC2 ?  can TBW see GASTCCompressor ? @todo Oodle fix me
 int32 GASTCCompressor = 0;
 static FAutoConsoleVariableRef CVarASTCCompressor(
 	TEXT("cook.ASTCTextureCompressor"),
@@ -39,25 +40,18 @@ static FAutoConsoleVariableRef CVarASTCCompressor(
 	TEXT("0: IntelISPC, 1: Arm"),
 	ECVF_Default | ECVF_ReadOnly
 );
-int32 GASTCHDRProfile = 0;
-static FAutoConsoleVariableRef CVarAllowASTCHDRProfile(
-	TEXT("cook.AllowASTCHDRProfile"),
-	GASTCHDRProfile,
-	TEXT("whether to allow ASTC HDR profile, the hdr format is only supported on some devices, e.g. Apple A13, Mali-G72, Adreno (TM) 660"),
-	ECVF_Default | ECVF_ReadOnly
-);
+
 #if PLATFORM_WINDOWS || PLATFORM_LINUX || PLATFORM_MAC
 	#define SUPPORTS_ISPC_ASTC	1
 #else
 	#define SUPPORTS_ISPC_ASTC	0
 #endif
 
-// increment this if you change anything that will affect compression in this file, including FORCED_NORMAL_MAP_COMPRESSION_SIZE_VALUE
+// increment this if you change anything that will affect compression in this file
 #define BASE_ASTC_FORMAT_VERSION 40
 
 #define MAX_QUALITY_BY_SIZE 4
 #define MAX_QUALITY_BY_SPEED 3
-#define FORCED_NORMAL_MAP_COMPRESSION_SIZE_VALUE 3
 
 
 DEFINE_LOG_CATEGORY_STATIC(LogTextureFormatASTC, Log, All);
@@ -85,6 +79,8 @@ class FASTCTextureBuildFunction final : public FTextureBuildFunction
 	op(ASTC_RGB) \
 	op(ASTC_RGBA) \
 	op(ASTC_RGBAuto) \
+	op(ASTC_RGBA_HQ) \
+	op(ASTC_RGB_HDR) \
 	op(ASTC_NormalAG) \
 	op(ASTC_NormalRG)
 
@@ -124,6 +120,7 @@ class FASTCTextureBuildFunction final : public FTextureBuildFunction
 
 static int32 GetDefaultCompressionBySizeValue(FCbObjectView InFormatConfigOverride)
 {
+	// this is code duped between TextureFormatASTC and TextureFormatISPC
 	if (InFormatConfigOverride)
 	{
 		// If we have an explicit format config, then use it directly
@@ -133,15 +130,25 @@ static int32 GetDefaultCompressionBySizeValue(FCbObjectView InFormatConfigOverri
 		checkf(!FieldView.HasError(), TEXT("Failed to parse DefaultASTCQualityBySize value from FormatConfigOverride"));
 		return CompressionModeValue;
 	}
+	else
+	{
+		// default of 0 == 12x12 ?
+		// BaseEngine.ini sets DefaultASTCQualityBySize to 3 == 6x6
 
-	// start at default quality, then lookup in .ini file
-	int32 CompressionModeValue = 0;
-	GConfig->GetInt(TEXT("/Script/UnrealEd.CookerSettings"), TEXT("DefaultASTCQualityBySize"), CompressionModeValue, GEngineIni);
-	
-	FParse::Value(FCommandLine::Get(), TEXT("-astcqualitybysize="), CompressionModeValue);
-	CompressionModeValue = FMath::Min<uint32>(CompressionModeValue, MAX_QUALITY_BY_SIZE);
+		auto GetCompressionModeValue = []() {
+			// start at default quality, then lookup in .ini file
+			int32 CompressionModeValue = 0;
+			GConfig->GetInt(TEXT("/Script/UnrealEd.CookerSettings"), TEXT("DefaultASTCQualityBySize"), CompressionModeValue, GEngineIni);
 	
-	return CompressionModeValue;
+			FParse::Value(FCommandLine::Get(), TEXT("-astcqualitybysize="), CompressionModeValue);
+			
+			return FMath::Min<uint32>(CompressionModeValue, MAX_QUALITY_BY_SIZE);
+		};
+
+		static int32 CompressionModeValue = GetCompressionModeValue();
+
+		return CompressionModeValue;
+	}
 }
 
 static int32 GetDefaultCompressionBySpeedValue(FCbObjectView InFormatConfigOverride)
@@ -155,32 +162,59 @@ static int32 GetDefaultCompressionBySpeedValue(FCbObjectView InFormatConfigOverr
 		checkf(!FieldView.HasError(), TEXT("Failed to parse DefaultASTCQualityBySpeed value from FormatConfigOverride"));
 		return CompressionModeValue;
 	}
+	else
+	{
 
-	// start at default quality, then lookup in .ini file
-	int32 CompressionModeValue = 0;
-	GConfig->GetInt(TEXT("/Script/UnrealEd.CookerSettings"), TEXT("DefaultASTCQualityBySpeed"), CompressionModeValue, GEngineIni);
-	
-	FParse::Value(FCommandLine::Get(), TEXT("-astcqualitybyspeed="), CompressionModeValue);
-	CompressionModeValue = FMath::Min<uint32>(CompressionModeValue, MAX_QUALITY_BY_SPEED);
+		// default of 0 == "fastest"
+
+		auto GetCompressionModeValue = []() {
+			// start at default quality, then lookup in .ini file
+			int32 CompressionModeValue = 0;
+			GConfig->GetInt(TEXT("/Script/UnrealEd.CookerSettings"), TEXT("DefaultASTCQualityBySpeed"), CompressionModeValue, GEngineIni);
 	
-	return CompressionModeValue;
+			FParse::Value(FCommandLine::Get(), TEXT("-astcqualitybyspeed="), CompressionModeValue);
+			
+			return FMath::Min<uint32>(CompressionModeValue, MAX_QUALITY_BY_SPEED);
+		};
+
+		static int32 CompressionModeValue = GetCompressionModeValue();
+
+		return CompressionModeValue;
+	}
 }
 
-static FString GetQualityString(const FCbObjectView& InFormatConfigOverride, int32 OverrideSizeValue=-1, int32 OverrideSpeedValue=-1)
+static FString GetQualityString(EPixelFormat PixelFormat,const FCbObjectView& InFormatConfigOverride)
 {
 	// convert to a string
 	FString CompressionMode;
-	switch (OverrideSizeValue >= 0 ? OverrideSizeValue : GetDefaultCompressionBySizeValue(InFormatConfigOverride))
-	{
-		case 0:	CompressionMode = TEXT("12x12"); break;
-		case 1:	CompressionMode = TEXT("10x10"); break;
-		case 2:	CompressionMode = TEXT("8x8"); break;
-		case 3:	CompressionMode = TEXT("6x6"); break;
-		case 4:	CompressionMode = TEXT("4x4"); break;
-		default: UE_LOG(LogTemp, Fatal, TEXT("ASTC size quality higher than expected"));
+
+	switch (PixelFormat)
+	{
+		case PF_ASTC_12x12:
+		case PF_ASTC_12x12_HDR:
+			CompressionMode = TEXT("12x12"); 
+			break;
+		case PF_ASTC_10x10:
+		case PF_ASTC_10x10_HDR:
+			CompressionMode = TEXT("10x10"); 
+			break;
+		case PF_ASTC_8x8:
+		case PF_ASTC_8x8_HDR:
+			CompressionMode = TEXT("8x8"); 
+			break;
+		case PF_ASTC_6x6:
+		case PF_ASTC_6x6_HDR:
+			CompressionMode = TEXT("6x6");
+			break;
+		case PF_ASTC_4x4:
+		case PF_ASTC_4x4_HDR:
+			CompressionMode = TEXT("4x4");
+			break;
+		default:
+			UE_LOG(LogTemp, Fatal, TEXT("ASTC size quality higher than expected"));
 	}
 	
-	switch (OverrideSpeedValue >= 0 ? OverrideSpeedValue : GetDefaultCompressionBySpeedValue(InFormatConfigOverride))
+	switch ( GetDefaultCompressionBySpeedValue(InFormatConfigOverride) )
 	{
 		case 0:	CompressionMode += TEXT(" -fastest"); break;
 		case 1:	CompressionMode += TEXT(" -fast"); break;
@@ -192,8 +226,32 @@ static FString GetQualityString(const FCbObjectView& InFormatConfigOverride, int
 	return CompressionMode;
 }
 
-static EPixelFormat GetQualityFormat(const FCbObjectView& InFormatConfigOverride, int32 OverrideSizeValue=-1, bool bHDRFormat = false)
+static EPixelFormat GetQualityFormat(const FTextureBuildSettings& BuildSettings)
 {
+	// code dupe between TextureFormatASTC  and TextureFormatISPC
+
+	const FCbObjectView& InFormatConfigOverride = BuildSettings.FormatConfigOverride;
+	int32 OverrideSizeValue= BuildSettings.CompressionQuality;
+
+	bool bIsNormalMap = (BuildSettings.TextureFormatName == GTextureFormatNameASTC_NormalAG || BuildSettings.TextureFormatName == GTextureFormatNameASTC_NormalRG);
+	bool bIsHQ = BuildSettings.TextureFormatName == GTextureFormatNameASTC_RGBA_HQ;
+	bool bHDRFormat = BuildSettings.TextureFormatName == GTextureFormatNameASTC_RGB_HDR;
+
+	if ( bIsNormalMap )
+	{
+		return PF_ASTC_6x6;
+	}
+	else if ( bIsHQ )
+	{
+		return PF_ASTC_4x4;
+	}
+	else if (BuildSettings.bVirtualStreamable)
+	{
+		return PF_ASTC_4x4;		
+	}
+
+	// CompressionQuality value here is ETextureCompressionQuality minus 1
+
 	// convert to a string
 	EPixelFormat Format = PF_Unknown;
 	if (bHDRFormat)
@@ -223,12 +281,6 @@ static EPixelFormat GetQualityFormat(const FCbObjectView& InFormatConfigOverride
 	return Format;
 }
 
-static uint16 GetQualityVersion(const FCbObjectView& InFormatConfigOverride, int32 OverrideSizeValue = -1)
-{
-	// top 3 bits for size compression value, and next 3 for speed
-	return ((OverrideSizeValue >= 0 ? OverrideSizeValue : GetDefaultCompressionBySizeValue(InFormatConfigOverride)) << 13) | (GetDefaultCompressionBySpeedValue(InFormatConfigOverride) << 10);
-}
-
 static bool CompressSliceToASTC(
 	const FImageView & SourceImage,
 	FString CompressionParameters,
@@ -236,6 +288,8 @@ static bool CompressSliceToASTC(
 	IImageWrapperModule& ImageWrapperModule
 	)
 {
+	// at this point, SourceImage has been converted to RGBA8 or RGBA16F based on whether
+	//	 the request TextureFormatName is "ASTC_RGB_HDR" or not, so we can ask if "source" is HDR :
 	bool bHDR = ERawImageFormat::IsHDR(SourceImage.Format);
 	
 	EImageFormat FileFormat = bHDR ? EImageFormat::EXR : EImageFormat::PNG;
@@ -443,59 +497,47 @@ public:
 		const FTextureBuildSettings* BuildSettings = nullptr
 	) const override
 	{
-	/*
-	// note/bug , doesn't call through to ISPC :
 #if SUPPORTS_ISPC_ASTC
 		if(GASTCCompressor == 0)
 		{
-			return IntelISPCTexCompFormat.GetVersion(Format,BuildSettings);
+			// set high bit so version numbers of ISPC and ASTC don't overlap :
+			check( BASE_ASTC_FORMAT_VERSION < 0x80 );
+			return 0x80 | IntelISPCTexCompFormat.GetVersion(Format,BuildSettings);
 		}
 #endif
-	*/
 
 		return BASE_ASTC_FORMAT_VERSION;
 	}
 
 	virtual FString GetDerivedDataKeyString(const FTextureBuildSettings& BuildSettings) const override
 	{
-	/*
-	// note/bug : doesn't call through to ISPC !
-	// but it is putting GASTCCompressor in the key
 #if SUPPORTS_ISPC_ASTC
 		if(GASTCCompressor == 0)
 		{
 			return IntelISPCTexCompFormat.GetDerivedDataKeyString(BuildSettings);
 		}
 #endif
-	*/
 
-		return FString::Printf(TEXT("ASTCCmpr_%d_%d_%d"), GetQualityVersion(BuildSettings.FormatConfigOverride, BuildSettings.CompressionQuality), GASTCCompressor, int32(GASTCHDRProfile > 0 && BuildSettings.bHDRSource));
+		// ASTC block size chosen is in PixelFormat
+		EPixelFormat PixelFormat = GetQualityFormat(BuildSettings);
+		int Speed = GetDefaultCompressionBySpeedValue(BuildSettings.FormatConfigOverride);
+
+ 		return FString::Printf(TEXT("ASTC_%d_%d"), (int)PixelFormat,Speed);
 	}
 
 	virtual FTextureFormatCompressorCaps GetFormatCapabilities() const override
 	{
-		FTextureFormatCompressorCaps RetCaps;
-		// use defaults
-		return RetCaps;
+		return FTextureFormatCompressorCaps(); // Default capabilities.
 	}
 
 	virtual void GetSupportedFormats(TArray<FName>& OutFormats) const override
 	{
-		for (int32 i = 0; i < UE_ARRAY_COUNT(GSupportedTextureFormatNames); ++i)
-		{
-			OutFormats.Add(GSupportedTextureFormatNames[i]);
-		}
+		OutFormats.Append(GSupportedTextureFormatNames, sizeof(GSupportedTextureFormatNames)/sizeof(GSupportedTextureFormatNames[0]) ); 
 	}
 
 	virtual EPixelFormat GetPixelFormatForImage(const FTextureBuildSettings& BuildSettings, const struct FImage& Image, bool bImageHasAlphaChannel) const override
 	{
-		// special case for normal maps
-		if (BuildSettings.TextureFormatName == GTextureFormatNameASTC_NormalAG || BuildSettings.TextureFormatName == GTextureFormatNameASTC_NormalRG)
-		{
-			return GetQualityFormat(BuildSettings.FormatConfigOverride, FORCED_NORMAL_MAP_COMPRESSION_SIZE_VALUE);
-		}
-		
-		return GetQualityFormat(BuildSettings.FormatConfigOverride, BuildSettings.CompressionQuality, Image.Format == ERawImageFormat::RGBA16F);
+		return GetQualityFormat(BuildSettings);
 	}
 
 	virtual bool CompressImage(
@@ -510,56 +552,75 @@ public:
 		if(GASTCCompressor == 0)
 		{
 			// Route ASTC compression work to the ISPC module instead.
+			// note: ISPC can't do HDR, will throw an error
 			return IntelISPCTexCompFormat.CompressImage(InImage, BuildSettings, DebugTexturePathName, bImageHasAlphaChannel, OutCompressedImage);
 		}
 #endif
 
-		// note: BuildSettings.bHDRSource is a weird thing to be checking here
-		//  it should look at whether the requested output format is HDR
-		//  the input->output mapping has already been done, you shouldn't do it again
-		bool bHDRImage = BuildSettings.bHDRSource && GASTCHDRProfile;
-		if ( BuildSettings.bHDRSource && ! GASTCHDRProfile )
-		{
-			UE_LOG(LogTextureFormatASTC, Display, TEXT("HDR Image passed to ASTC but ASTC HDR is disabled, will compress as 8 bit: (%s)"),*FString(DebugTexturePathName));
-		}
+		bool bHDRImage = BuildSettings.TextureFormatName == GTextureFormatNameASTC_RGB_HDR;
 
 		// Get Raw Image Data from passed in FImage & convert to BGRA8 or RGBA16F
 		// note: wasteful, often copies image to same format
 		FImage Image;
 		InImage.CopyTo(Image, bHDRImage ? ERawImageFormat::RGBA16F : ERawImageFormat::BGRA8, BuildSettings.GetDestGammaSpace());
 
+		// @todo Oodle : for HDR need to clamp the F16 image in [0,F16_max] (like Oodle Texture does)
+		//	  maybe also set A to 1.f to match BC6 ?
+
 		// Determine the compressed pixel format and compression parameters
 		EPixelFormat CompressedPixelFormat = GetPixelFormatForImage(BuildSettings, Image, bImageHasAlphaChannel);
 
 		FString CompressionParameters = TEXT("");
 
-		bool bIsRGBColor = (BuildSettings.TextureFormatName == GTextureFormatNameASTC_RGB ||
-			((BuildSettings.TextureFormatName == GTextureFormatNameASTC_RGBAuto) && !bImageHasAlphaChannel));
-		bool bIsRGBAColor = (BuildSettings.TextureFormatName == GTextureFormatNameASTC_RGBA ||
-			((BuildSettings.TextureFormatName == GTextureFormatNameASTC_RGBAuto) && bImageHasAlphaChannel));
+		FString QualityString = GetQualityString(CompressedPixelFormat,BuildSettings.FormatConfigOverride);
+		
+		/*
+
+		The modes available are:
+
+			-cl : use the linear LDR color profile.
+			-cs : use the sRGB LDR color profile.
+			-ch : use the HDR color profile, tuned for HDR RGB and LDR A.
+			-cH : use the HDR color profile, tuned for HDR RGBA.
+
+			-ch or -cl is added later in CompressSlice
+
+		*/
 
-		if (bIsRGBColor)
+
+		if (bHDRImage)
+		{
+			CompressionParameters = FString::Printf(TEXT("%s"), *QualityString );
+		}
+		else if ( BuildSettings.TextureFormatName == GTextureFormatNameASTC_RGB ||
+			BuildSettings.TextureFormatName == GTextureFormatNameASTC_RGBA || 
+			BuildSettings.TextureFormatName == GTextureFormatNameASTC_RGBAuto || 
+			BuildSettings.TextureFormatName == GTextureFormatNameASTC_RGBA_HQ )
 		{
-			if (bHDRImage)
+
+			if ( BuildSettings.TextureFormatName == GTextureFormatNameASTC_RGB ||
+				! bImageHasAlphaChannel )
 			{
-				CompressionParameters = FString::Printf(TEXT("%s %s -cw 1 1 1 0"), *GetQualityString(BuildSettings.FormatConfigOverride, BuildSettings.CompressionQuality), /*BuildSettings.bSRGB ? TEXT("-srgb") :*/ TEXT(""));
+				// even if Name was RGBA we still use the RGB profile if !bImageHasAlphaChannel
+				//	so that "Compress Without Alpha" can force us to opaque
+
+				// we need to set alpha to opaque here
+				// can do it using "1" in the bgra swizzle to astcenc
+				
+				CompressionParameters = FString::Printf(TEXT("%s -esw bgr1"), *QualityString );
 			}
 			else
 			{
-				CompressionParameters = FString::Printf(TEXT("%s %s -esw bgra -cw 1 1 1 0"), *GetQualityString(BuildSettings.FormatConfigOverride, BuildSettings.CompressionQuality), /*BuildSettings.bSRGB ? TEXT("-srgb") :*/ TEXT("") );
+				CompressionParameters = FString::Printf(TEXT("%s -esw bgra"), *QualityString );
 			}
 		}
-		else if (bIsRGBAColor)
-		{
-			CompressionParameters = FString::Printf(TEXT("%s %s -esw bgra -cw 1 1 1 1"), *GetQualityString(BuildSettings.FormatConfigOverride, BuildSettings.CompressionQuality), /*BuildSettings.bSRGB ? TEXT("-srgb") :*/ TEXT("") );
-		}
 		else if (BuildSettings.TextureFormatName == GTextureFormatNameASTC_NormalAG)
 		{
-			CompressionParameters = FString::Printf(TEXT("%s -esw 0g0b -cw 0 1 0 1 -dblimit 60 -b 2.5 -v 3 1 1 0 50 0 -va 1 1 0 50"), *GetQualityString(BuildSettings.FormatConfigOverride, FORCED_NORMAL_MAP_COMPRESSION_SIZE_VALUE, -1));
+			CompressionParameters = FString::Printf(TEXT("%s -esw 0g0b -cw 0 1 0 1 -dblimit 60 -b 2.5 -v 3 1 1 0 50 0 -va 1 1 0 50"), *QualityString);
 		}
 		else if (BuildSettings.TextureFormatName == GTextureFormatNameASTC_NormalRG)
 		{
-			CompressionParameters = FString::Printf(TEXT("%s -esw bg00 -cw 1 1 0 0 -dblimit 60 -b 2.5 -v 3 1 1 0 50 0 -va 1 1 0 50"), *GetQualityString(BuildSettings.FormatConfigOverride, FORCED_NORMAL_MAP_COMPRESSION_SIZE_VALUE, -1));
+			CompressionParameters = FString::Printf(TEXT("%s -esw bg00 -cw 1 1 0 0 -dblimit 60 -b 2.5 -v 3 1 1 0 50 0 -va 1 1 0 50"), *QualityString);
 		}
 		else
 		{
diff --git a/Engine/Source/Developer/TextureFormatETC2/Private/TextureFormatETC2.cpp b/Engine/Source/Developer/TextureFormatETC2/Private/TextureFormatETC2.cpp
--- a/Engine/Source/Developer/TextureFormatETC2/Private/TextureFormatETC2.cpp
+++ b/Engine/Source/Developer/TextureFormatETC2/Private/TextureFormatETC2.cpp
@@ -192,15 +192,20 @@ public:
 	virtual EPixelFormat GetPixelFormatForImage(const struct FTextureBuildSettings& BuildSettings, const struct FImage& Image, bool bImageHasAlphaChannel) const override
 	{
 		if (BuildSettings.TextureFormatName == GTextureFormatNameETC2_RGB ||
-			(BuildSettings.TextureFormatName == GTextureFormatNameAutoETC2 && !bImageHasAlphaChannel))
+			BuildSettings.TextureFormatName == GTextureFormatNameETC2_RGBA ||
+			BuildSettings.TextureFormatName == GTextureFormatNameAutoETC2 )
 		{
-			return PF_ETC2_RGB;
-		}
-
-		if (BuildSettings.TextureFormatName == GTextureFormatNameETC2_RGBA ||
-				(BuildSettings.TextureFormatName == GTextureFormatNameAutoETC2 && bImageHasAlphaChannel))
-		{
-			return PF_ETC2_RGBA;
+			if ( BuildSettings.TextureFormatName == GTextureFormatNameETC2_RGB || !bImageHasAlphaChannel )
+			{
+				// even if Name was RGBA we still use the RGB profile if !bImageHasAlphaChannel
+				//	so that "Compress Without Alpha" can force us to opaque
+
+				return PF_ETC2_RGB;
+			}
+			else
+			{
+				return PF_ETC2_RGBA;
+			}
 		}
 
 		if (BuildSettings.TextureFormatName == GTextureFormatNameETC2_R11)
diff --git a/Engine/Source/Developer/TextureFormatIntelISPCTexComp/Private/TextureFormatIntelISPCTexComp.cpp b/Engine/Source/Developer/TextureFormatIntelISPCTexComp/Private/TextureFormatIntelISPCTexComp.cpp
--- a/Engine/Source/Developer/TextureFormatIntelISPCTexComp/Private/TextureFormatIntelISPCTexComp.cpp
+++ b/Engine/Source/Developer/TextureFormatIntelISPCTexComp/Private/TextureFormatIntelISPCTexComp.cpp
@@ -40,7 +40,7 @@ class FIntelISPCTexCompTextureBuildFunction final : public FTextureBuildFunction
 	}
 };
 
-// increment this if you change anything that will affect compression in this file, including FORCED_NORMAL_MAP_COMPRESSION_SIZE_VALUE
+// increment this if you change anything that will affect compression in this file
 #define BASE_ISPC_DX11_FORMAT_VERSION 4
 
 // For debugging intermediate image results by saving them out as files.
@@ -65,10 +65,18 @@ class FIntelISPCTexCompTextureBuildFunction final : public FTextureBuildFunction
 	#undef DECL_FORMAT_NAME_ENTRY
 #undef ENUM_SUPPORTED_FORMATS
 
+/**
+ * note : ASTC formats are NOT in the GSupportedTextureFormatNames
+ *  so we are not registered to support them
+ *  they will call into TextureFormatASTC and then optionally redirect to here
+ * this just defines the FNames
+ */
 #define ENUM_ASTC_FORMATS(op) \
 	op(ASTC_RGB) \
 	op(ASTC_RGBA) \
 	op(ASTC_RGBAuto) \
+	op(ASTC_RGBA_HQ) \
+	op(ASTC_RGB_HDR) \
 	op(ASTC_NormalAG) \
 	op(ASTC_NormalRG)
 
@@ -138,7 +146,7 @@ struct FBitmapInfoHeader
 void SaveImageAsBMP( FArchive& Ar, const uint8* RawData, int SourceBytesPerPixel, int SizeX, int SizeY )
 {
 	FBitmapFileHeader bmf;
-	FBitmapInfoHeader bmhdr;
+	FBitmapInfoHeader bmh;
 
 	// File header.
 	bmf.bfType = 'B' + (256 * (int32)'M');
@@ -146,23 +154,23 @@ void SaveImageAsBMP( FArchive& Ar, const uint8* RawData, int SourceBytesPerPixel
 	bmf.bfReserved2 = 0;
 	int32 biSizeImage = SizeX * SizeY * 3;
 	bmf.bfOffBits = sizeof(FBitmapFileHeader) + sizeof(FBitmapInfoHeader);
-	bmhdr.biBitCount = 24;
+	bmh.biBitCount = 24;
 
 	bmf.bfSize = bmf.bfOffBits + biSizeImage;
 	Ar << bmf;
 
 	// Info header.
-	bmhdr.biSize = sizeof(FBitmapInfoHeader);
-	bmhdr.biWidth = SizeX;
-	bmhdr.biHeight = SizeY;
-	bmhdr.biPlanes = 1;
-	bmhdr.biCompression = BCBI_RGB;
-	bmhdr.biSizeImage = biSizeImage;
-	bmhdr.biXPelsPerMeter = 0;
-	bmhdr.biYPelsPerMeter = 0;
-	bmhdr.biClrUsed = 0;
-	bmhdr.biClrImportant = 0;
-	Ar << bmhdr;
+	bmh.biSize = sizeof(FBitmapInfoHeader);
+	bmh.biWidth = SizeX;
+	bmh.biHeight = SizeY;
+	bmh.biPlanes = 1;
+	bmh.biCompression = BCBI_RGB;
+	bmh.biSizeImage = biSizeImage;
+	bmh.biXPelsPerMeter = 0;
+	bmh.biYPelsPerMeter = 0;
+	bmh.biClrUsed = 0;
+	bmh.biClrImportant = 0;
+	Ar << bmh;
 
 	bool bIsRGBA16 = (SourceBytesPerPixel == 8);
 
@@ -397,10 +405,10 @@ static void IntelBC7CompressScans(bc7_enc_settings* pEncSettings, FImage* pInIma
 }
 
 #define MAX_QUALITY_BY_SIZE 4
-#define FORCED_NORMAL_MAP_COMPRESSION_SIZE_VALUE 3
 
 static uint16 GetDefaultCompressionBySizeValue(FCbObjectView InFormatConfigOverride)
 {
+	// this is code duped between TextureFormatASTC and TextureFormatISPC
 	if (InFormatConfigOverride)
 	{
 		// If we have an explicit format config, then use it directly
@@ -410,21 +418,51 @@ static uint16 GetDefaultCompressionBySizeValue(FCbObjectView InFormatConfigOverr
 		checkf(!FieldView.HasError(), TEXT("Failed to parse DefaultASTCQualityBySize value from FormatConfigOverride"));
 		return CompressionModeValue;
 	}
+	else
+	{
+		// default of 0 == 12x12 ?
+		// BaseEngine.ini sets DefaultASTCQualityBySize to 3 == 6x6
 
-	// start at default quality, then lookup in .ini file
-	int32 CompressionModeValue = 0;
-	GConfig->GetInt(TEXT("/Script/UnrealEd.CookerSettings"), TEXT("DefaultASTCQualityBySize"), CompressionModeValue, GEngineIni);
+		auto GetCompressionModeValue = []() {
+			// start at default quality, then lookup in .ini file
+			int32 CompressionModeValue = 0;
+			GConfig->GetInt(TEXT("/Script/UnrealEd.CookerSettings"), TEXT("DefaultASTCQualityBySize"), CompressionModeValue, GEngineIni);
+	
+			FParse::Value(FCommandLine::Get(), TEXT("-astcqualitybysize="), CompressionModeValue);
+			
+			return FMath::Min<uint32>(CompressionModeValue, MAX_QUALITY_BY_SIZE);
+		};
 
-	FParse::Value(FCommandLine::Get(), TEXT("-astcqualitybysize="), CompressionModeValue);
-	CompressionModeValue = FMath::Min<uint32>(CompressionModeValue, MAX_QUALITY_BY_SIZE);
+		static int32 CompressionModeValue = GetCompressionModeValue();
 
-	return CompressionModeValue;
+		return CompressionModeValue;
+	}
 }
 
-static EPixelFormat GetQualityFormat(int& BlockWidth, int& BlockHeight, const FCbObjectView& InFormatConfigOverride, int32 OverrideSizeValue = -1)
+static EPixelFormat GetASTCQualityFormat(int& BlockWidth, int& BlockHeight, const FTextureBuildSettings& BuildSettings)
 {
+	// code dupe between TextureFormatASTC  and TextureFormatISPC
+
+	const FCbObjectView& InFormatConfigOverride = BuildSettings.FormatConfigOverride;
+	int32 OverrideSizeValue= BuildSettings.CompressionQuality;
+
+	bool bIsNormalMap = (BuildSettings.TextureFormatName == GTextureFormatNameASTC_NormalAG || BuildSettings.TextureFormatName == GTextureFormatNameASTC_NormalRG);
+	bool bIsHQ = BuildSettings.TextureFormatName == GTextureFormatNameASTC_RGBA_HQ;
+	bool bHDRFormat = BuildSettings.TextureFormatName == GTextureFormatNameASTC_RGB_HDR;
+	check( ! bHDRFormat ); // HDR not supported here
+	
+	if ( bIsNormalMap )
+	{
+		BlockWidth = BlockHeight = 6;
+		return PF_ASTC_6x6;
+	}
+	else if ( bIsHQ || BuildSettings.bVirtualStreamable )
+	{
+		BlockWidth = BlockHeight = 4;
+		return PF_ASTC_4x4;		
+	}
+
 	// Note: ISPC only supports 8x8 and higher quality, and only one speed (fast)
-	// convert to a string
 	EPixelFormat Format = PF_Unknown;
 	switch (OverrideSizeValue >= 0 ? OverrideSizeValue : GetDefaultCompressionBySizeValue(InFormatConfigOverride))
 	{
@@ -467,6 +505,9 @@ static void IntelASTCCompressScans(FASTCEncoderSettings* pEncSettings, FImage* p
 	uint8* pInTexels = reinterpret_cast<uint8*>(&pInImage->RawData[0]) + InSliceSize * SliceIndex;
 	uint8* pOutTexels = reinterpret_cast<uint8*>(&pOutImage->RawData[0]) + OutSliceSize * SliceIndex;
 
+	// TextureFormatName at this point should no longer be RGBAuto
+	//	it's mapped to _RGB or _RGBA concretely
+
 	if (pEncSettings->TextureFormatName == GTextureFormatNameASTC_RGB)
 	{
 		// Switch byte order for compressors input (BGRA -> RGBA)
@@ -485,7 +526,7 @@ static void IntelASTCCompressScans(FASTCEncoderSettings* pEncSettings, FImage* p
 			}
 		}
 	}
-	else if (pEncSettings->TextureFormatName == GTextureFormatNameASTC_RGBA)
+	else if (pEncSettings->TextureFormatName == GTextureFormatNameASTC_RGBA || pEncSettings->TextureFormatName == GTextureFormatNameASTC_RGBA_HQ)
 	{
 		// Switch byte order for compressors input (BGRA -> RGBA)
 		for (int y = yStart; y < yEnd; ++y)
@@ -635,16 +676,13 @@ public:
 		return BASE_ISPC_DX11_FORMAT_VERSION;
 	}
 	
-	/*
 	virtual FString GetDerivedDataKeyString(const FTextureBuildSettings& BuildSettings) const override
 	{
-		// @todo Oodle : should log settings here :
-
-		// should store
-		// GetDefaultCompressionBySizeValue() 
-		//return FString::Printf(TEXT("ISPC_%d_%d_%d"), GetQualityVersion(BuildSettings.FormatConfigOverride, BuildSettings.CompressionQuality), GASTCCompressor, int32(GASTCHDRProfile > 0 && BuildSettings.bHDRSource));
+		// ASTC block size chosen is in PixelFormat
+		EPixelFormat PixelFormat = GetPixelFormatForBuildSettings(BuildSettings);
+		
+ 		return FString::Printf(TEXT("ISPC_%d"), (int)PixelFormat);
 	}
-	*/
 
 	virtual void GetSupportedFormats(TArray<FName>& OutFormats) const override
 	{
@@ -753,13 +791,20 @@ public:
 			{
 				F16Value.Encoded = 0;
 			}
-			// @todo Oodle : should also clamp +Inf to max non-inf value?
+			// @todo Oodle : should also clamp +Inf to max non-inf value? -> yes, like Oodle Texture does
+			//	  maybe also set A to 1.f to match BC6 ?
+
 			//	 note that behavior of the basic F32->F16 conversion was changed,
 			//	 it used to clamp to max val if out of range, now stores inf
 		}
 	}
 
 	virtual EPixelFormat GetPixelFormatForImage(const FTextureBuildSettings& BuildSettings, const struct FImage& Image, bool bImageHasAlphaChannel) const override
+	{
+		return GetPixelFormatForBuildSettings(BuildSettings);
+	}
+
+	EPixelFormat GetPixelFormatForBuildSettings(const FTextureBuildSettings& BuildSettings) const
 	{
 		if (BuildSettings.TextureFormatName == GTextureFormatNameBC6H)
 		{
@@ -769,17 +814,11 @@ public:
 		{
 			return PF_BC7;
 		}
-		else if (BuildSettings.bVirtualStreamable)
-		{
-			return PF_ASTC_4x4;
-		}
 		else
 		{
 			int _Width, _Height;
-			bool bIsNormalMap = (BuildSettings.TextureFormatName == GTextureFormatNameASTC_NormalAG ||
-				BuildSettings.TextureFormatName == GTextureFormatNameASTC_NormalRG);
 
-			return GetQualityFormat(_Width, _Height, BuildSettings.FormatConfigOverride, bIsNormalMap ? FORCED_NORMAL_MAP_COMPRESSION_SIZE_VALUE : BuildSettings.CompressionQuality);
+			return GetASTCQualityFormat(_Width, _Height, BuildSettings);
 		}
 	}
 
@@ -847,61 +886,43 @@ public:
 		}
 		else
 		{
-			bool bIsRGBColorASTC = (BuildSettings.TextureFormatName == GTextureFormatNameASTC_RGB ||
-				((BuildSettings.TextureFormatName == GTextureFormatNameASTC_RGBAuto) && !bImageHasAlphaChannel));
-			bool bIsRGBAColorASTC = (BuildSettings.TextureFormatName == GTextureFormatNameASTC_RGBA ||
-				((BuildSettings.TextureFormatName == GTextureFormatNameASTC_RGBAuto) && bImageHasAlphaChannel));
 			bool bIsNormalMap = (BuildSettings.TextureFormatName == GTextureFormatNameASTC_NormalAG ||
 				BuildSettings.TextureFormatName == GTextureFormatNameASTC_NormalRG);
-
-			if (BuildSettings.bVirtualStreamable)
-			{
-				// Always use 4x4 for streamable VT, to reduce texture format fragmentation
-				BlockWidth = 4;
-				BlockHeight = 4;
-			}
-			else
-			{
-				 GetQualityFormat( BlockWidth, BlockHeight, BuildSettings.FormatConfigOverride, bIsNormalMap ? FORCED_NORMAL_MAP_COMPRESSION_SIZE_VALUE : BuildSettings.CompressionQuality );
-			}
-			check(CompressedPixelFormat == PF_ASTC_4x4 || !BuildSettings.bVirtualStreamable);
+				
+			bool bHDRImage = BuildSettings.TextureFormatName == GTextureFormatNameASTC_RGB_HDR;
+			check( ! bHDRImage ); // HDR not supported here
+			
+			GetASTCQualityFormat( BlockWidth, BlockHeight, BuildSettings );
 
 			FASTCEncoderSettings EncoderSettings;
-			if (BuildSettings.TextureFormatName == GTextureFormatNameASTC_NormalAG)
-			{
-				GetProfile_astc_alpha_fast(&EncoderSettings, BlockWidth, BlockHeight);
-				EncoderSettings.TextureFormatName = BuildSettings.TextureFormatName;
-				bCompressionSucceeded = true;
-				check(EncoderSettings.block_width!=0);
-			}
-			else if (BuildSettings.TextureFormatName == GTextureFormatNameASTC_NormalRG)
+			EncoderSettings.TextureFormatName = BuildSettings.TextureFormatName;
+			
+			if ( bIsNormalMap )
 			{
 				GetProfile_astc_fast(&EncoderSettings, BlockWidth, BlockHeight);
-				EncoderSettings.TextureFormatName = BuildSettings.TextureFormatName;
-				bCompressionSucceeded = true;
-				check(EncoderSettings.block_width!=0);
 			}
-			else if (bIsRGBColorASTC)
-			{
-				GetProfile_astc_fast(&EncoderSettings, BlockWidth, BlockHeight);
-				EncoderSettings.TextureFormatName = GTextureFormatNameASTC_RGB;
-				bCompressionSucceeded = true;
-				check(EncoderSettings.block_width!=0);
+			else
+			{			
+				if( BuildSettings.TextureFormatName == GTextureFormatNameASTC_RGB || !bImageHasAlphaChannel )
+				{
+					// even if Name was RGBA we still use the RGB profile if !bImageHasAlphaChannel
+					//	so that "Compress Without Alpha" can force us to opaque
+
+					EncoderSettings.TextureFormatName = GTextureFormatNameASTC_RGB; // <- this will force A to 255 in IntelASTCCompressScans
+
+					GetProfile_astc_fast(&EncoderSettings, BlockWidth, BlockHeight);
+				}
+				else
+				{
+					EncoderSettings.TextureFormatName = GTextureFormatNameASTC_RGBA;
+					
+					GetProfile_astc_alpha_fast(&EncoderSettings, BlockWidth, BlockHeight);
+				}
 			}
-			else if (bIsRGBAColorASTC)
+			
 			{
-				GetProfile_astc_alpha_fast(&EncoderSettings, BlockWidth, BlockHeight);
-				EncoderSettings.TextureFormatName = GTextureFormatNameASTC_RGBA;
-				bCompressionSucceeded = true;
 				check(EncoderSettings.block_width!=0);
-			}
-			else
-			{
-				check(0);
-			}
 
-			if (bCompressionSucceeded)
-			{
 				FImage Image;
 				InImage.CopyTo(Image, ERawImageFormat::BGRA8, BuildSettings.GetDestGammaSpace());
 
@@ -925,6 +946,8 @@ public:
 
 				FMultithreadedCompression<FASTCEncoderSettings>::Compress(MultithreadSettings, EncoderSettings, Image, OutCompressedImage, &IntelASTCCompressScans, bUseTasks);
 
+				bCompressionSucceeded = true;
+				
 #if DEBUG_SAVE_INTERMEDIATE_IMAGES
 				//@DEBUG (save swizzled/fixed-up input as BMP):
 				if (SaveInputOutput)	// && LocalCounter < 10 && Image.SizeX >= 1024)
diff --git a/Engine/Source/Runtime/Engine/Classes/Engine/Texture.h b/Engine/Source/Runtime/Engine/Classes/Engine/Texture.h
--- a/Engine/Source/Runtime/Engine/Classes/Engine/Texture.h
+++ b/Engine/Source/Runtime/Engine/Classes/Engine/Texture.h
@@ -84,11 +84,11 @@ UENUM()
 enum ETextureCompressionQuality
 {
 	TCQ_Default = 0		UMETA(DisplayName="Default"),
-	TCQ_Lowest = 1		UMETA(DisplayName="Lowest"),
-	TCQ_Low = 2			UMETA(DisplayName="Low"),
-	TCQ_Medium = 3		UMETA(DisplayName="Medium"),
-	TCQ_High= 4			UMETA(DisplayName="High"),
-	TCQ_Highest = 5		UMETA(DisplayName="Highest"),
+	TCQ_Lowest = 1		UMETA(DisplayName="Lowest (ASTC 12x12)"),
+	TCQ_Low = 2			UMETA(DisplayName="Low (ASTC 10x10)"),
+	TCQ_Medium = 3		UMETA(DisplayName="Medium (ASTC 8x8)"),
+	TCQ_High= 4			UMETA(DisplayName="High (ASTC 6x6)"),
+	TCQ_Highest = 5		UMETA(DisplayName="Highest (ASTC 4x4)"),
 	TCQ_MAX,
 };
 
@@ -1145,7 +1145,7 @@ public:
 	UPROPERTY(EditAnywhere, BlueprintReadWrite, Category=Adjustments, meta=(DisplayName = "Max Alpha"))
 	float AdjustMaxAlpha;
 
-	/** If enabled, the texture's alpha channel will be discarded during compression */
+	/** If enabled, the texture's alpha channel will be forced to opaque for any compressed texture output format.  Does not apply if output format is uncompressed RGBA. */
 	UPROPERTY(EditAnywhere, BlueprintReadWrite, Category=Compression, meta=(DisplayName="Compress Without Alpha"))
 	uint32 CompressionNoAlpha:1;
 
diff --git a/Engine/Source/Runtime/Engine/Private/Texture.cpp b/Engine/Source/Runtime/Engine/Private/Texture.cpp
--- a/Engine/Source/Runtime/Engine/Private/Texture.cpp
+++ b/Engine/Source/Runtime/Engine/Private/Texture.cpp
@@ -2678,6 +2678,8 @@ void UTexture::GetDefaultFormatSettings(FTextureFormatSettings& OutSettings) con
 	OutSettings.CompressionSettings = CompressionSettings;
 	OutSettings.CompressionNone = CompressionNone;
 	OutSettings.CompressionNoAlpha = CompressionNoAlpha;
+	// Texture does not have CompressionForceAlpha
+	OutSettings.CompressionForceAlpha = false;
 	OutSettings.CompressionYCoCg = CompressionYCoCg;
 	OutSettings.SRGB = SRGB;
 }
@@ -2704,6 +2706,9 @@ void UTexture::SetLayerFormatSettings(int32 LayerIndex, const FTextureFormatSett
 		CompressionSettings = InSettings.CompressionSettings;
 		CompressionNone = InSettings.CompressionNone;
 		CompressionNoAlpha = InSettings.CompressionNoAlpha;
+		// CompressionForceAlpha not copied
+		//CompressionForceAlpha = InSettings.CompressionForceAlpha;
+		check( ! InSettings.CompressionForceAlpha );
 		CompressionYCoCg = InSettings.CompressionYCoCg;
 		SRGB = InSettings.SRGB;
 	}
@@ -2866,7 +2871,8 @@ static FName ConditionalGetPrefixedFormat(FName TextureFormatName, const ITarget
 	return TextureFormatName;
 }
 
-FName GetDefaultTextureFormatName( const ITargetPlatform* TargetPlatform, const UTexture* Texture, int32 LayerIndex, bool bSupportCompressedVolumeTexture, int32 BlockSize, bool bSupportFilteredFloat32Textures )
+FName GetDefaultTextureFormatName( const ITargetPlatform* TargetPlatform, const UTexture* Texture, int32 LayerIndex, 
+	bool bSupportCompressedVolumeTexture, int32 Unused_BlockSize, bool bSupportFilteredFloat32Textures )
 {
 	FName TextureFormatName = NAME_None;
 	bool bOodleTextureSdkVersionIsNone = true;
@@ -2878,7 +2884,6 @@ FName GetDefaultTextureFormatName( const ITargetPlatform* TargetPlatform, const
 #if WITH_EDITOR
 	// Supported texture format names.
 	static FName NameDXT1(TEXT("DXT1"));
-	static FName NameDXT3(TEXT("DXT3"));
 	static FName NameDXT5(TEXT("DXT5"));
 	static FName NameDXT5n(TEXT("DXT5n"));
 	static FName NameAutoDXT(TEXT("AutoDXT"));
@@ -2957,6 +2962,7 @@ FName GetDefaultTextureFormatName( const ITargetPlatform* TargetPlatform, const
 		// Also force uncompressed if size of top mip is not a multiple of 4
 		// note that even if top mip is a multiple of 4, lower may not be
 		// we can only choose compression if it's supported by all platforms/RHI's (else check TargetPlatform->SupportsFeature)
+		// note: does not use the passed-in "BlockSize" parameter, hard coded to 4
 		if ( (SizeX < 4) || (SizeY < 4) || (SizeX % 4 != 0) || (SizeY % 4 != 0) )
 		{
 			// don't log if TC was going to map to uncompressed anyway
@@ -3037,9 +3043,7 @@ FName GetDefaultTextureFormatName( const ITargetPlatform* TargetPlatform, const
 		}
 		else
 		{
-			// note CompressionNoAlpha no longer kills alpha on AutoDXT
-			//  CompressionNoAlpha really just means "pick BC1 instead of BC3"
-			//	but if your image is not a multiple of 4 you still get alpha RGBA
+			// note CompressionNoAlpha no longer kills alpha if it's forced to uncompressed (eg. because size is not multiple of 4)
 			TextureFormatName = NameBGRA8;
 		}
 	}
@@ -3113,9 +3117,12 @@ FName GetDefaultTextureFormatName( const ITargetPlatform* TargetPlatform, const
 	{
 		check( FormatSettings.CompressionSettings == TC_Default ||
 			FormatSettings.CompressionSettings == TC_Masks ); 
+
 		if (FormatSettings.CompressionNoAlpha)
 		{
-				// CompressionNoAlpha only affects AutoDXT
+			// CompressionNoAlpha changes AutoDXT to DXT1 early
+			//	this is unnecessary/redundant, I believe
+			//  the later handling of AutoDXT would make this same mapping
 			TextureFormatName = NameDXT1;
 		}
 		else
@@ -3168,13 +3175,14 @@ FName GetDefaultTextureFormatName( const ITargetPlatform* TargetPlatform, const
 	return ConditionalGetPrefixedFormat(TextureFormatName, TargetPlatform, bOodleTextureSdkVersionIsNone);
 }
 
-void GetDefaultTextureFormatNamePerLayer(TArray<FName>& OutFormatNames, const class ITargetPlatform* TargetPlatform, const class UTexture* Texture, bool bSupportCompressedVolumeTexture, int32 BlockSize, bool bSupportFilteredFloat32Textures )
+void GetDefaultTextureFormatNamePerLayer(TArray<FName>& OutFormatNames, const class ITargetPlatform* TargetPlatform, const class UTexture* Texture, 
+	bool bSupportCompressedVolumeTexture, int32 Unused_BlockSize, bool bSupportFilteredFloat32Textures )
 {
 #if WITH_EDITOR
 	OutFormatNames.Reserve(Texture->Source.GetNumLayers());
 	for (int32 LayerIndex = 0; LayerIndex < Texture->Source.GetNumLayers(); ++LayerIndex)
 	{
-		OutFormatNames.Add(GetDefaultTextureFormatName(TargetPlatform, Texture, LayerIndex, bSupportCompressedVolumeTexture, BlockSize, bSupportFilteredFloat32Textures));
+		OutFormatNames.Add(GetDefaultTextureFormatName(TargetPlatform, Texture, LayerIndex, bSupportCompressedVolumeTexture, Unused_BlockSize, bSupportFilteredFloat32Textures));
 	}
 #endif // WITH_EDITOR
 }
@@ -3185,7 +3193,6 @@ void GetAllDefaultTextureFormats(const class ITargetPlatform* TargetPlatform, TA
 
 #if WITH_EDITOR
 	static FName NameDXT1(TEXT("DXT1"));
-	static FName NameDXT3(TEXT("DXT3"));
 	static FName NameDXT5(TEXT("DXT5"));
 	static FName NameDXT5n(TEXT("DXT5n"));
 	static FName NameAutoDXT(TEXT("AutoDXT"));
@@ -3204,7 +3211,6 @@ void GetAllDefaultTextureFormats(const class ITargetPlatform* TargetPlatform, TA
 	static FName NameBC7(TEXT("BC7"));
 
 	OutFormats.Add(NameDXT1);
-	OutFormats.Add(NameDXT3);
 	OutFormats.Add(NameDXT5);
 	OutFormats.Add(NameDXT5n);
 	OutFormats.Add(NameAutoDXT);
diff --git a/Engine/Source/Runtime/Engine/Public/TextureResource.h b/Engine/Source/Runtime/Engine/Public/TextureResource.h
--- a/Engine/Source/Runtime/Engine/Public/TextureResource.h
+++ b/Engine/Source/Runtime/Engine/Public/TextureResource.h
@@ -605,10 +605,10 @@ private:
 };
 
 /** Gets the name of a format for the given LayerIndex */
-ENGINE_API FName GetDefaultTextureFormatName( const class ITargetPlatform* TargetPlatform, const class UTexture* Texture, int32 LayerIndex, bool bSupportCompressedVolumeTexture = false, int32 BlockSize = 4, bool bSupportFilteredFloat32Textures = true);
+ENGINE_API FName GetDefaultTextureFormatName( const class ITargetPlatform* TargetPlatform, const class UTexture* Texture, int32 LayerIndex, bool bSupportCompressedVolumeTexture, int32 Unused_BlockSize, bool bSupportFilteredFloat32Textures);
 
 /** Gets an array of format names for each layer in the texture */
-ENGINE_API void GetDefaultTextureFormatNamePerLayer(TArray<FName>& OutFormatNames, const class ITargetPlatform* TargetPlatform, const class UTexture* Texture, bool bSupportCompressedVolumeTexture = false, int32 BlockSize = 4, bool bSupportFilteredFloat32Textures = true);
+ENGINE_API void GetDefaultTextureFormatNamePerLayer(TArray<FName>& OutFormatNames, const class ITargetPlatform* TargetPlatform, const class UTexture* Texture, bool bSupportCompressedVolumeTexture, int32 Unused_BlockSize, bool bSupportFilteredFloat32Textures);
 
 // returns all the texture formats which can be returned by GetDefaultTextureFormatName
 ENGINE_API void GetAllDefaultTextureFormats( const class ITargetPlatform* TargetPlatform, TArray<FName>& OutFormats);
