TextureFormatASTC better logs

diff --git a/Engine/Source/Developer/TextureFormatASTC/Private/TextureFormatASTC.cpp b/Engine/Source/Developer/TextureFormatASTC/Private/TextureFormatASTC.cpp
--- a/Engine/Source/Developer/TextureFormatASTC/Private/TextureFormatASTC.cpp
+++ b/Engine/Source/Developer/TextureFormatASTC/Private/TextureFormatASTC.cpp
@@ -211,7 +211,7 @@ static FString GetQualityString(EPixelFormat PixelFormat,const FCbObjectView& In
 			CompressionMode = TEXT("4x4");
 			break;
 		default:
-			UE_LOG(LogTemp, Fatal, TEXT("ASTC size quality higher than expected"));
+			UE_LOG(LogTextureFormatASTC, Fatal, TEXT("ASTC size quality higher than expected"));
 	}
 	
 	switch ( GetDefaultCompressionBySpeedValue(InFormatConfigOverride) )
@@ -220,7 +220,7 @@ static FString GetQualityString(EPixelFormat PixelFormat,const FCbObjectView& In
 		case 1:	CompressionMode += TEXT(" -fast"); break;
 		case 2:	CompressionMode += TEXT(" -medium"); break;
 		case 3:	CompressionMode += TEXT(" -thorough"); break;
-		default: UE_LOG(LogTemp, Fatal, TEXT("ASTC speed quality higher than expected"));
+		default: UE_LOG(LogTextureFormatASTC, Fatal, TEXT("ASTC speed quality higher than expected"));
 	}
 
 	return CompressionMode;
@@ -263,7 +263,7 @@ static EPixelFormat GetQualityFormat(const FTextureBuildSettings& BuildSettings)
 			case 2:	Format = PF_ASTC_8x8_HDR; break;
 			case 3:	Format = PF_ASTC_6x6_HDR; break;
 			case 4:	Format = PF_ASTC_4x4_HDR; break;
-			default: UE_LOG(LogTemp, Fatal, TEXT("Max quality higher than expected"));
+			default: UE_LOG(LogTextureFormatASTC, Fatal, TEXT("Max quality higher than expected"));
 		}
 	}
 	else
@@ -275,7 +275,7 @@ static EPixelFormat GetQualityFormat(const FTextureBuildSettings& BuildSettings)
 			case 2:	Format = PF_ASTC_8x8; break;
 			case 3:	Format = PF_ASTC_6x6; break;
 			case 4:	Format = PF_ASTC_4x4; break;
-			default: UE_LOG(LogTemp, Fatal, TEXT("Max quality higher than expected"));
+			default: UE_LOG(LogTextureFormatASTC, Fatal, TEXT("Max quality higher than expected"));
 		}
 	}
 	return Format;
@@ -551,12 +551,20 @@ public:
 #if SUPPORTS_ISPC_ASTC
 		if(GASTCCompressor == 0)
 		{
+			UE_CALL_ONCE( [&](){
+				UE_LOG(LogTextureFormatASTC, Display, TEXT("TextureFormatASTC using ISPC"))
+			} );
+
 			// Route ASTC compression work to the ISPC module instead.
 			// note: ISPC can't do HDR, will throw an error
 			return IntelISPCTexCompFormat.CompressImage(InImage, BuildSettings, DebugTexturePathName, bImageHasAlphaChannel, OutCompressedImage);
 		}
 #endif
 
+		UE_CALL_ONCE( [&](){
+			UE_LOG(LogTextureFormatASTC, Display, TEXT("TextureFormatASTC using astcenc"))
+		} );
+
 		bool bHDRImage = BuildSettings.TextureFormatName == GTextureFormatNameASTC_RGB_HDR;
 
 		// Get Raw Image Data from passed in FImage & convert to BGRA8 or RGBA16F
