fix astcenc swapping RB in TextureFormatASTC
also new cvars to run the decoder and leave temp files :
cook.ASTCDebugLeaveTempFiles=1
cook.ASTCDebugWriteDecodedImage=1

diff --git a/Engine/Source/Developer/TextureFormatASTC/Private/TextureFormatASTC.cpp b/Engine/Source/Developer/TextureFormatASTC/Private/TextureFormatASTC.cpp
--- a/Engine/Source/Developer/TextureFormatASTC/Private/TextureFormatASTC.cpp
+++ b/Engine/Source/Developer/TextureFormatASTC/Private/TextureFormatASTC.cpp
@@ -41,6 +41,25 @@ static FAutoConsoleVariableRef CVarASTCCompressor(
 	ECVF_Default | ECVF_ReadOnly
 );
 
+// turn on to leave temp files in Intermediate/Cache for debugging
+static int32 GASTCDebugLeaveTempFiles = 0;
+static FAutoConsoleVariableRef CVarDebugLeaveTempFiles(
+	TEXT("cook.ASTCDebugLeaveTempFiles"),
+	GASTCDebugLeaveTempFiles,
+	TEXT("0: default, 1: leave debug temp files in Intermediate/Cache"),
+	ECVF_Default | ECVF_ReadOnly
+);
+
+// turn on to write decoded image in Intermediate/Cache for debugging
+static int32 GASTCDebugWriteDecodedImage = 0;
+static FAutoConsoleVariableRef CVarDebugWriteDecodedImage(
+	TEXT("cook.ASTCDebugWriteDecodedImage"),
+	GASTCDebugWriteDecodedImage,
+	TEXT("0: default, 1: write decoded image in Intermediate/Cache"),
+	ECVF_Default | ECVF_ReadOnly
+);
+
+
 #if PLATFORM_WINDOWS || PLATFORM_LINUX || PLATFORM_MAC
 	#define SUPPORTS_ISPC_ASTC	1
 #else
@@ -48,7 +67,7 @@ static FAutoConsoleVariableRef CVarASTCCompressor(
 #endif
 
 // increment this if you change anything that will affect compression in this file
-#define BASE_ASTC_FORMAT_VERSION 43
+#define BASE_ASTC_FORMAT_VERSION 44
 
 #define MAX_QUALITY_BY_SIZE 4
 #define MAX_QUALITY_BY_SPEED 3
@@ -306,11 +325,9 @@ static bool CompressSliceToASTC(
 	// make a random file name to write the image :
 	FGuid Guid;
 	FPlatformMisc::CreateGuid(Guid);
-	FString InputFilePath = FString::Printf(TEXT("Cache/%08x-%08x-%08x-%08x-RGBToASTCIn"), Guid.A, Guid.B, Guid.C, Guid.D) + TEXT(".") + ImageWrapperModule.GetExtension(FileFormat);
-	FString OutputFilePath = FString::Printf(TEXT("Cache/%08x-%08x-%08x-%08x-RGBToASTCOut.astc"), Guid.A, Guid.B, Guid.C, Guid.D);
-
-	InputFilePath  = FPaths::ProjectIntermediateDir() + InputFilePath;
-	OutputFilePath = FPaths::ProjectIntermediateDir() + OutputFilePath;
+	FString BaseFilePath = FPaths::ProjectIntermediateDir() + FString::Printf(TEXT("Cache/TFASTC-%08x-%08x-%08x-%08x-"), Guid.A, Guid.B, Guid.C, Guid.D);
+	FString InputFilePath = BaseFilePath + TEXT("In.") + ImageWrapperModule.GetExtension(FileFormat);
+	FString OutputFilePath = BaseFilePath + TEXT("Out.astc");
 
 	// write to InputFilePath :
 	{
@@ -333,6 +350,8 @@ static bool CompressSliceToASTC(
 
 	// FileData written, can free now :
 	FileData.Reset();
+	
+	// @@CB @todo Oodle : why do we use -cl (LDR linear) and not -cs (LDR sRGB) ?
 
 	// Compress PNG file to ASTC (using the reference astcenc.exe from ARM)
 	// -j option to set thread count ? when we're running lots of textures at the same time in a cook,
@@ -448,13 +467,51 @@ static bool CompressSliceToASTC(
 		check(ASTCData.Num() == (sizeof(FASTCHeader) + MipSize));
 		FMemory::Memcpy(MipData, ASTCData.GetData() + sizeof(FASTCHeader), MipSize);
 	}
+
+	if ( GASTCDebugWriteDecodedImage )
+	{	
+		FString DecodedFilePath = BaseFilePath + TEXT("Dec.") + ImageWrapperModule.GetExtension(FileFormat);
+
+		// Params starts with -cl or -ch or -cs , grab that character and change to -dl etc :
+		check( Params[0] == TEXT('-') );
+		check( Params[1] == TEXT('c') );
+
+		FString DecoderParams = TEXT("-d");
+		DecoderParams += Params[2];
+		DecoderParams += FString::Printf(TEXT(" \"%s\" \"%s\""),*OutputFilePath,*DecodedFilePath);
 		
+		UE_LOG(LogTextureFormatASTC, Verbose, TEXT("Decoding ASTC (encode options = '%s' , decode = '%s')..."), *CompressionParameters, *DecoderParams);
+
+		FProcHandle Proc = FPlatformProcess::CreateProc(*CompressorPath, *DecoderParams, true, false, false, NULL, -1, NULL, NULL);
+
+		// Failed to start the compressor process
+		if (!Proc.IsValid())
+		{
+			UE_LOG(LogTextureFormatASTC, Error, TEXT("Failed to start astcenc for decompressing images (%s)"), *CompressorPath);
+		}
+		else
+		{
+			FPlatformProcess::WaitForProc(Proc);
+			FPlatformProcess::CloseProc(Proc);
+
+			// right after we make it, delete it to clean up
+			// break point here to examine
+			// or turn on GASTCDebugLeaveTempFiles
+
+			if ( ! GASTCDebugLeaveTempFiles )
+			{
+				IFileManager::Get().Delete(*DecodedFilePath);
+			}
+		}
+	}
+
 	// Delete intermediate files
+	if ( ! GASTCDebugLeaveTempFiles )
 	{
 		TRACE_CPUPROFILER_EVENT_SCOPE(ASTC.DeleteFiles);
 
-	IFileManager::Get().Delete(*InputFilePath);
-	IFileManager::Get().Delete(*OutputFilePath);
+		IFileManager::Get().Delete(*InputFilePath);
+		IFileManager::Get().Delete(*OutputFilePath);
 	}
 
 	return true;
@@ -613,7 +670,7 @@ public:
 
 		if (bHDRImage)
 		{
-			CompressionParameters = FString::Printf(TEXT("%s"), *QualityString );
+			CompressionParameters = QualityString;
 			
 			// ASTC can encode floats that BC6H can't
 			//  but still clamp as if we were BC6H, so that the same output is made
@@ -636,22 +693,22 @@ public:
 				// we need to set alpha to opaque here
 				// can do it using "1" in the bgra swizzle to astcenc
 				
-				CompressionParameters = FString::Printf(TEXT("%s -esw bgr1"), *QualityString );
+				CompressionParameters = FString::Printf(TEXT("%s -esw rgb1"), *QualityString );
 			}
 			else
 			{
-				CompressionParameters = FString::Printf(TEXT("%s -esw bgra"), *QualityString );
+				CompressionParameters = QualityString;
 			}
 		}
 		else if (BuildSettings.TextureFormatName == GTextureFormatNameASTC_NormalAG)
 		{
-			// astcenc has "-normal" , dunno if that does anything other than just channel swizzling (it outputs xxxy)
-
-			CompressionParameters = FString::Printf(TEXT("%s -esw 0g0b -cw 0 1 0 1 -dblimit 60"), *QualityString);
+			// or "gggr" ?
+			// note that DXT5n processing does "1g0r"
+			CompressionParameters = FString::Printf(TEXT("%s -esw 0g0r -cw 0 1 0 1 -dblimit 60"), *QualityString);
 		}
 		else if (BuildSettings.TextureFormatName == GTextureFormatNameASTC_NormalRG)
 		{
-			CompressionParameters = FString::Printf(TEXT("%s -esw bg00 -cw 1 1 0 0 -dblimit 60"), *QualityString);
+			CompressionParameters = FString::Printf(TEXT("%s -esw rg00 -cw 1 1 0 0 -dblimit 60"), *QualityString);
 		}
 		else
 		{
