astcenc use correct gamma space encoder mode
TextureFormatASTC

diff --git a/Engine/Source/Developer/TextureFormatASTC/Private/TextureFormatASTC.cpp b/Engine/Source/Developer/TextureFormatASTC/Private/TextureFormatASTC.cpp
--- a/Engine/Source/Developer/TextureFormatASTC/Private/TextureFormatASTC.cpp
+++ b/Engine/Source/Developer/TextureFormatASTC/Private/TextureFormatASTC.cpp
@@ -67,7 +67,7 @@ static FAutoConsoleVariableRef CVarDebugWriteDecodedImage(
 #endif
 
 // increment this if you change anything that will affect compression in this file
-#define BASE_ASTC_FORMAT_VERSION 44
+#define BASE_ASTC_FORMAT_VERSION 45
 
 #define MAX_QUALITY_BY_SIZE 4
 #define MAX_QUALITY_BY_SPEED 3
@@ -304,7 +304,8 @@ static bool CompressSliceToASTC(
 	const FImageView & SourceImage,
 	FString CompressionParameters,
 	TArray64<uint8>& OutCompressedData,
-	IImageWrapperModule& ImageWrapperModule
+	IImageWrapperModule& ImageWrapperModule,
+	bool bSRGB
 	)
 {
 	// at this point, SourceImage has been converted to RGBA8 or RGBA16F based on whether
@@ -351,18 +352,45 @@ static bool CompressSliceToASTC(
 	// FileData written, can free now :
 	FileData.Reset();
 	
-	// @@CB @todo Oodle : why do we use -cl (LDR linear) and not -cs (LDR sRGB) ?
-
 	// Compress PNG file to ASTC (using the reference astcenc.exe from ARM)
-	// -j option to set thread count ? when we're running lots of textures at the same time in a cook,
-	//	 it might be better to limit the astcenc process to fewer threads?
-	FString Params = (bHDR ? TEXT("-ch ") : TEXT("-cl ")) + FString::Printf(TEXT("\"%s\" \"%s\" %s -silent"),
+	
+	/*
+
+	The modes available are:
+
+		-cl : use the linear LDR color profile.
+		-cs : use the sRGB LDR color profile.
+		-ch : use the HDR color profile, tuned for HDR RGB and LDR A.
+		-cH : use the HDR color profile, tuned for HDR RGBA.
+
+	*/
+
+	const TCHAR * CompressionMode;
+	if ( bHDR )
+	{
+		CompressionMode = TEXT("-ch");
+	}
+	else if ( bSRGB )
+	{	
+		CompressionMode = TEXT("-cs");
+	}
+	else
+	{
+		CompressionMode = TEXT("-cl");
+	}
+
+	FString Params = FString::Printf(TEXT("%s \"%s\" \"%s\" %s -silent"),
+		CompressionMode,
 		*InputFilePath,
 		*OutputFilePath,
 		*CompressionParameters
 	);
+	
+	// maybe?
+	// could use -j option to set thread count ? when we're running lots of textures at the same time in a cook,
+	//	 it might be better to limit the astcenc process to fewer threads?
 
-	UE_LOG(LogTextureFormatASTC, Verbose, TEXT("Compressing to ASTC (options = '%s')..."), *CompressionParameters);
+	UE_LOG(LogTextureFormatASTC, Verbose, TEXT("Compressing to ASTC (astcenc %s)..."), *Params);
 
 	// Start Compressor
 	// @todo Oodle : check if we have sse4 and use those
@@ -641,6 +669,9 @@ public:
 		} );
 
 		bool bHDRImage = BuildSettings.TextureFormatName == GTextureFormatNameASTC_RGB_HDR;
+		// DestGamma is how the texture will be bound to GPU
+		bool bSRGB = BuildSettings.GetDestGammaSpace() == EGammaSpace::sRGB;
+		check( !bHDRImage || !bSRGB );
 
 		// Get Raw Image Data from passed in FImage & convert to BGRA8 or RGBA16F
 		// note: wasteful, often copies image to same format
@@ -654,20 +685,6 @@ public:
 
 		FString QualityString = GetQualityString(CompressedPixelFormat,BuildSettings.FormatConfigOverride);
 		
-		/*
-
-		The modes available are:
-
-			-cl : use the linear LDR color profile.
-			-cs : use the sRGB LDR color profile.
-			-ch : use the HDR color profile, tuned for HDR RGB and LDR A.
-			-cH : use the HDR color profile, tuned for HDR RGBA.
-
-			-ch or -cl is added later in CompressSlice
-
-		*/
-
-
 		if (bHDRImage)
 		{
 			CompressionParameters = QualityString;
@@ -724,7 +741,7 @@ public:
 
 			FImageView Slice = Image.GetSlice(SliceIndex);
 			
-			bCompressionSucceeded = CompressSliceToASTC(Slice,CompressionParameters,CompressedSliceData,ImageWrapperModule);
+			bCompressionSucceeded = CompressSliceToASTC(Slice,CompressionParameters,CompressedSliceData,ImageWrapperModule,bSRGB);
 
 			if ( ! bCompressionSucceeded )
 			{
