Use ARM astcenc by default for ASTC texture compression

diff --git a/Engine/Source/Developer/TextureFormatASTC/Private/TextureFormatASTC.cpp b/Engine/Source/Developer/TextureFormatASTC/Private/TextureFormatASTC.cpp
--- a/Engine/Source/Developer/TextureFormatASTC/Private/TextureFormatASTC.cpp
+++ b/Engine/Source/Developer/TextureFormatASTC/Private/TextureFormatASTC.cpp
@@ -29,13 +29,13 @@
 * 
 * TextureFormatASTC runs the ARM astcenc
 * 
-* or (by default) redirects to Intel ISPC texcomp* 
+* or redirects to Intel ISPC texcomp* 
 * 
 *****/
 
 // when GASTCCompressor == 0 ,use TextureFormatIntelISPCTexComp instead of this
 // @todo Oodle : GASTCCompressor global breaks DDC2.  Need to pass through so TBW can see.
-int32 GASTCCompressor = 0;
+int32 GASTCCompressor = 1;
 static FAutoConsoleVariableRef CVarASTCCompressor(
 	TEXT("cook.ASTCTextureCompressor"),
 	GASTCCompressor,
