TextureFormatASTC fix build

diff --git a/Engine/Source/Developer/TextureFormatASTC/Private/TextureFormatASTC.cpp b/Engine/Source/Developer/TextureFormatASTC/Private/TextureFormatASTC.cpp
--- a/Engine/Source/Developer/TextureFormatASTC/Private/TextureFormatASTC.cpp
+++ b/Engine/Source/Developer/TextureFormatASTC/Private/TextureFormatASTC.cpp
@@ -250,7 +250,7 @@ static EPixelFormat GetQualityFormat(const FTextureBuildSettings& BuildSettings)
 
 
 static bool ASTCEnc_Compress(
-	FImage& InImage,
+	const FImage& InImage,
 	const FTextureBuildSettings& BuildSettings,
 	const FIntVector3& InMip0Dimensions,
 	int32 InMip0NumSlicesNoDepth,
