Readded "TextureFormatIntelISPCTexComp" so "TextureFormatIntelISPCTexComp" is loaded before ASTC.

diff --git a/Engine/Source/Developer/TextureFormatASTC/TextureFormatASTC.Build.cs b/Engine/Source/Developer/TextureFormatASTC/TextureFormatASTC.Build.cs
--- a/Engine/Source/Developer/TextureFormatASTC/TextureFormatASTC.Build.cs
+++ b/Engine/Source/Developer/TextureFormatASTC/TextureFormatASTC.Build.cs
@@ -19,6 +19,7 @@ public class TextureFormatASTC : ModuleRules
 			"ImageCore",
 			"ImageWrapper",
 			"TextureBuild",
+			"TextureFormatIntelISPCTexComp",
 			"astcenc",
 		});
 	}
