Add a new set of ASTC pixel formats for storing high precision normals in RG channels. The texture is encoded as LA for higher precision, swizzled to RG at runtime for shader compatibility with BC5. Requires RHI support.
Runtime pixel formats are named  PF_ASTC_<blocksize>_NORM_RG. Use texture format name ASTC_NormalRG_Precise for encoding.

diff --git a/Engine/Plugins/Tests/RHITests/Source/RHITests/Public/RHITextureTests.h b/Engine/Plugins/Tests/RHITests/Source/RHITests/Public/RHITextureTests.h
--- a/Engine/Plugins/Tests/RHITests/Source/RHITests/Public/RHITextureTests.h
+++ b/Engine/Plugins/Tests/RHITests/Source/RHITests/Public/RHITextureTests.h
@@ -716,6 +716,11 @@ public:
 			case PF_R32G32B32_UINT:
 			case PF_R32G32B32_SINT:
 			case PF_R9G9B9EXP5:
+			case PF_ASTC_4x4_NORM_RG:
+			case PF_ASTC_6x6_NORM_RG:
+			case PF_ASTC_8x8_NORM_RG:
+			case PF_ASTC_10x10_NORM_RG:
+			case PF_ASTC_12x12_NORM_RG:
 				break;
 
 			case PF_G8:
diff --git a/Engine/Source/Developer/TextureFormatASTC/Private/TextureFormatASTC.cpp b/Engine/Source/Developer/TextureFormatASTC/Private/TextureFormatASTC.cpp
--- a/Engine/Source/Developer/TextureFormatASTC/Private/TextureFormatASTC.cpp
+++ b/Engine/Source/Developer/TextureFormatASTC/Private/TextureFormatASTC.cpp
@@ -85,7 +85,8 @@ class FASTCTextureBuildFunction final : public FTextureBuildFunction
 	op(ASTC_RGB_HDR) \
 	op(ASTC_NormalLA) \
 	op(ASTC_NormalAG) \
-	op(ASTC_NormalRG)
+	op(ASTC_NormalRG) \
+	op(ASTC_NormalRG_Precise) // Encoded as LA for precision, mapped to RG at runtime. RHI needs to support PF_ASTC_*_NORM_RG formats (requires runtime swizzle)
 
 	#define DECL_FORMAT_NAME(FormatName) static FName GTextureFormatName##FormatName = FName(TEXT(#FormatName));
 		ENUM_SUPPORTED_FORMATS(DECL_FORMAT_NAME);
@@ -123,9 +124,10 @@ class FASTCTextureBuildFunction final : public FTextureBuildFunction
 static bool IsNormalMapFormat(FName TextureFormatName)
 {
 	return
-		TextureFormatName == GTextureFormatNameASTC_NormalAG || 
-		TextureFormatName == GTextureFormatNameASTC_NormalRG || 
-		TextureFormatName == GTextureFormatNameASTC_NormalLA;
+		TextureFormatName == GTextureFormatNameASTC_NormalAG ||
+		TextureFormatName == GTextureFormatNameASTC_NormalRG ||
+		TextureFormatName == GTextureFormatNameASTC_NormalLA ||
+		TextureFormatName == GTextureFormatNameASTC_NormalRG_Precise;
 }
 
 static int32 GetDefaultCompressionBySizeValue(FCbObjectView InFormatConfigOverride)
@@ -206,6 +208,10 @@ static EPixelFormat GetQualityFormat(const FTextureBuildSettings& BuildSettings)
 
 	if ( bIsNormalMap )
 	{
+		if ( BuildSettings.TextureFormatName == GTextureFormatNameASTC_NormalRG_Precise )
+		{
+			return PF_ASTC_6x6_NORM_RG;
+		}
 		return PF_ASTC_6x6;
 	}
 	else if ( bIsHQ )
@@ -373,7 +379,7 @@ static bool ASTCEnc_Compress(
 		EncConfig.cw_b_weight = 0.0f;
 		EncConfig.cw_a_weight = 0.0f;
 	}
-	else if (BuildSettings.TextureFormatName == GTextureFormatNameASTC_NormalLA)
+	else if (BuildSettings.TextureFormatName == GTextureFormatNameASTC_NormalLA || BuildSettings.TextureFormatName == GTextureFormatNameASTC_NormalRG_Precise)
 	{
 		// L+A mode: rrrg
 		EncSwizzle.r = ASTCENC_SWZ_B;
diff --git a/Engine/Source/Runtime/Apple/MetalRHI/Private/MetalRHI.cpp b/Engine/Source/Runtime/Apple/MetalRHI/Private/MetalRHI.cpp
--- a/Engine/Source/Runtime/Apple/MetalRHI/Private/MetalRHI.cpp
+++ b/Engine/Source/Runtime/Apple/MetalRHI/Private/MetalRHI.cpp
@@ -761,7 +761,12 @@ FMetalDynamicRHI::FMetalDynamicRHI(ERHIFeatureLevel::Type RequestedFeatureLevel)
 	GMetalBufferFormats[PF_R64_UINT				] = { mtlpp::PixelFormat::Invalid, (uint8)EMetalBufferFormat::Unknown };
 	GMetalBufferFormats[PF_R9G9B9EXP5			] = { mtlpp::PixelFormat::Invalid, (uint8)EMetalBufferFormat::Unknown };
 	GMetalBufferFormats[PF_P010					] = { mtlpp::PixelFormat::Invalid, (uint8)EMetalBufferFormat::Unknown };
-	static_assert(PF_MAX == 87, "Please setup GMetalBufferFormats properly for the new pixel format");
+	GMetalBufferFormats[PF_ASTC_4x4_NORM_RG		] = { mtlpp::PixelFormat::Invalid, (uint8)EMetalBufferFormat::Unknown };
+	GMetalBufferFormats[PF_ASTC_6x6_NORM_RG		] = { mtlpp::PixelFormat::Invalid, (uint8)EMetalBufferFormat::Unknown };
+	GMetalBufferFormats[PF_ASTC_8x8_NORM_RG		] = { mtlpp::PixelFormat::Invalid, (uint8)EMetalBufferFormat::Unknown };
+	GMetalBufferFormats[PF_ASTC_10x10_NORM_RG	] = { mtlpp::PixelFormat::Invalid, (uint8)EMetalBufferFormat::Unknown };
+	GMetalBufferFormats[PF_ASTC_12x12_NORM_RG	] = { mtlpp::PixelFormat::Invalid, (uint8)EMetalBufferFormat::Unknown };
+	static_assert(PF_MAX == 92, "Please setup GMetalBufferFormats properly for the new pixel format");
 
 	// Initialize the platform pixel format map.
 	GPixelFormats[PF_Unknown			].PlatformFormat	= (uint32)mtlpp::PixelFormat::Invalid;
diff --git a/Engine/Source/Runtime/Core/Private/Misc/PixelFormat.cpp b/Engine/Source/Runtime/Core/Private/Misc/PixelFormat.cpp
--- a/Engine/Source/Runtime/Core/Private/Misc/PixelFormat.cpp
+++ b/Engine/Source/Runtime/Core/Private/Misc/PixelFormat.cpp
@@ -137,6 +137,13 @@ FPixelFormatInfo    GPixelFormats[PF_MAX] =
 	// P010 contains 2 textures: R16 luminance plane followed by R16G16 1/4 size chrominance plane. (upper 10 bits used)
 	// BlockSize/BlockBytes/NumComponents values don't make much sense for this format, so set them all to one.
 	FPixelFormatInfo(PF_P010,				TEXT("P010"),				   1,		   1,		   1,		   2,		   1,			   0),
+
+	// ASTC high precision NormalRG support
+	FPixelFormatInfo(PF_ASTC_4x4_NORM_RG,   TEXT("ASTC_4x4_NORM_RG"),      4,          4,          1,          16,         2,              0),
+	FPixelFormatInfo(PF_ASTC_6x6_NORM_RG,   TEXT("ASTC_6x6_NORM_RG"),      6,          6,          1,          16,         2,              0),
+	FPixelFormatInfo(PF_ASTC_8x8_NORM_RG,   TEXT("ASTC_8x8_NORM_RG"),      8,          8,          1,          16,         2,              0),
+	FPixelFormatInfo(PF_ASTC_10x10_NORM_RG, TEXT("ASTC_10x10_NORM_RG"),    10,         10,         1,          16,         2,              0),
+	FPixelFormatInfo(PF_ASTC_12x12_NORM_RG, TEXT("ASTC_12x12_NORM_RG"),    12,         12,         1,          16,         2,              0),
 };
 
 
diff --git a/Engine/Source/Runtime/Core/Public/PixelFormat.h b/Engine/Source/Runtime/Core/Public/PixelFormat.h
--- a/Engine/Source/Runtime/Core/Public/PixelFormat.h
+++ b/Engine/Source/Runtime/Core/Public/PixelFormat.h
@@ -101,7 +101,12 @@ enum EPixelFormat : uint8
 	PF_R64_UINT				=84,
 	PF_R9G9B9EXP5			=85,
 	PF_P010					=86,
-	PF_MAX					=87,
+	PF_ASTC_4x4_NORM_RG		=87, // RG format stored in LA endpoints for better precision (requires RHI support for texture swizzle)
+	PF_ASTC_6x6_NORM_RG		=88,	
+	PF_ASTC_8x8_NORM_RG		=89,	
+	PF_ASTC_10x10_NORM_RG	=90,	
+	PF_ASTC_12x12_NORM_RG	=91,	
+	PF_MAX					=92,
 };
 #define FOREACH_ENUM_EPIXELFORMAT(op) \
 	op(PF_Unknown) \
@@ -188,7 +193,12 @@ enum EPixelFormat : uint8
 	op(PF_R32G32B32_SINT) \
 	op(PF_R32G32B32F) \
 	op(PF_R8_SINT) \
-	op(PF_R64_UINT)
+	op(PF_R64_UINT) \
+	op(PF_ASTC_4x4_NORM_RG) \
+	op(PF_ASTC_6x6_NORM_RG) \
+	op(PF_ASTC_8x8_NORM_RG) \
+	op(PF_ASTC_10x10_NORM_RG) \
+	op(PF_ASTC_12x12_NORM_RG)
 
 // Defines which channel is valid for each pixel format
 enum class EPixelFormatChannelFlags : uint8
diff --git a/Engine/Source/Runtime/CoreUObject/Public/UObject/NoExportTypes.h b/Engine/Source/Runtime/CoreUObject/Public/UObject/NoExportTypes.h
--- a/Engine/Source/Runtime/CoreUObject/Public/UObject/NoExportTypes.h
+++ b/Engine/Source/Runtime/CoreUObject/Public/UObject/NoExportTypes.h
@@ -297,6 +297,11 @@ enum EPixelFormat : int
 	PF_R64_UINT,
 	PF_R9G9B9EXP5,
 	PF_P010,
+	PF_ASTC_4x4_NORM_RG,
+	PF_ASTC_6x6_NORM_RG,
+	PF_ASTC_8x8_NORM_RG,
+	PF_ASTC_10x10_NORM_RG,
+	PF_ASTC_12x12_NORM_RG,
 	PF_MAX,
 };
 
diff --git a/Engine/Source/Runtime/RenderCore/Private/RenderUtils.cpp b/Engine/Source/Runtime/RenderCore/Private/RenderUtils.cpp
--- a/Engine/Source/Runtime/RenderCore/Private/RenderUtils.cpp
+++ b/Engine/Source/Runtime/RenderCore/Private/RenderUtils.cpp
@@ -341,6 +341,11 @@ EPixelFormatChannelFlags GetPixelFormatValidChannels(EPixelFormat InPixelFormat)
 		EPixelFormatChannelFlags::R,		// PF_R64_UINT
 		EPixelFormatChannelFlags::RGB,		// PF_R9G9B9EXP5
 		EPixelFormatChannelFlags::None,		// PF_P010
+		EPixelFormatChannelFlags::RG,		// PF_ASTC_4x4_NORM_RG
+		EPixelFormatChannelFlags::RG,		// PF_ASTC_6x6_NORM_RG
+		EPixelFormatChannelFlags::RG,		// PF_ASTC_8x8_NORM_RG
+		EPixelFormatChannelFlags::RG,		// PF_ASTC_10x10_NORM_RG
+		EPixelFormatChannelFlags::RG,		// PF_ASTC_12x12_NORM_RG
 	};
 	static_assert(UE_ARRAY_COUNT(PixelFormatToChannelFlags) == (uint8)PF_MAX, "Missing pixel format");
 	return (InPixelFormat < PF_MAX) ? PixelFormatToChannelFlags[(uint8)InPixelFormat] : EPixelFormatChannelFlags::None;
diff --git a/Engine/Source/Runtime/Renderer/Private/SystemTextures.cpp b/Engine/Source/Runtime/Renderer/Private/SystemTextures.cpp
--- a/Engine/Source/Runtime/Renderer/Private/SystemTextures.cpp
+++ b/Engine/Source/Runtime/Renderer/Private/SystemTextures.cpp
@@ -1204,7 +1204,7 @@ template<typename TInType>
 void InitializeData(const TInType& InData, EPixelFormat InFormat, uint8* OutData, uint32& OutByteCount)
 {
 	// If a new format is added insure that it is either supported here, or at least flagged as not supported
-	static_assert(PF_MAX == 87);
+	static_assert(PF_MAX == 92);
 
 	switch (InFormat)
 	{
@@ -1298,6 +1298,11 @@ void InitializeData(const TInType& InData, EPixelFormat InFormat, uint8* OutData
 		case PF_ASTC_8x8_HDR:
 		case PF_ASTC_10x10_HDR:
 		case PF_ASTC_12x12_HDR:
+		case PF_ASTC_4x4_NORM_RG:
+		case PF_ASTC_6x6_NORM_RG:
+		case PF_ASTC_8x8_NORM_RG:
+		case PF_ASTC_10x10_NORM_RG:
+		case PF_ASTC_12x12_NORM_RG:
 		case PF_BC6H:
 		case PF_BC7:
 		case PF_XGXR8:
