TextureFormatASTC : make BC7 HQ mapping configurable with new config DefaultASTCQualityBySizeHQ and respect per-texture quality override
only affects astenc not ISPC

diff --git a/Engine/Config/BaseEngine.ini b/Engine/Config/BaseEngine.ini
--- a/Engine/Config/BaseEngine.ini
+++ b/Engine/Config/BaseEngine.ini
@@ -3119,7 +3119,10 @@ NDKAPILevel=android-26
 
 [/Script/UnrealEd.CookerSettings]
 DefaultASTCQualityBySpeed=2
+; DefaultASTCQualityBySize 3 means 6x6
 DefaultASTCQualityBySize=3
+; DefaultASTCQualityBySizeHQ 4 means 4x4
+DefaultASTCQualityBySizeHQ=4
 +ClassesExcludedOnDedicatedServer=WidgetBlueprint
 +ClassesExcludedOnDedicatedServer=GroupActor
 +ClassesExcludedOnDedicatedServer=MetaData
diff --git a/Engine/Source/Developer/DeveloperToolSettings/Classes/CookerSettings.h b/Engine/Source/Developer/DeveloperToolSettings/Classes/CookerSettings.h
--- a/Engine/Source/Developer/DeveloperToolSettings/Classes/CookerSettings.h
+++ b/Engine/Source/Developer/DeveloperToolSettings/Classes/CookerSettings.h
@@ -173,12 +173,16 @@ public:
 	TArray<FString> VersionedIntRValues;
 
 	/** Quality of 0 means fastest, 3 means best quality */
-	UPROPERTY(GlobalConfig, EditAnywhere, Category = Textures, meta = (DisplayName = "ASTC Compression Quality vs Speed (0-3, 0 is fastest)"))
+	UPROPERTY(GlobalConfig, EditAnywhere, Category = Textures, meta = (DisplayName = "ASTC Compression Speed (0-3, 0 is fastest)"))
 	int32 DefaultASTCQualityBySpeed;
 
 	/** Quality of 0 means smallest (12x12 block size), 4 means best (4x4 block size) */
-	UPROPERTY(GlobalConfig, EditAnywhere, Category = Textures, meta = (DisplayName = "ASTC Compression Quality vs Size (0-4, 0 is smallest)"))
+	UPROPERTY(GlobalConfig, EditAnywhere, Category = Textures, meta = (DisplayName = "ASTC Compression Block Size (0-4, 0 is smallest)"))
 	int32 DefaultASTCQualityBySize;
+	
+	/** Quality of 0 means smallest (12x12 block size), 4 means best (4x4 block size) */
+	UPROPERTY(GlobalConfig, EditAnywhere, Category = Textures, meta = (DisplayName = "ASTC Compression Block Size for High Quality (BC7) (0-4, 0 is smallest)"))
+	int32 DefaultASTCQualityBySizeHQ;
 
 	/** Allows opening cooked assets in the editor */
 	UPROPERTY(GlobalConfig, EditAnywhere, Category = Textures, meta = (
diff --git a/Engine/Source/Developer/DeveloperToolSettings/Private/CookerSettings.cpp b/Engine/Source/Developer/DeveloperToolSettings/Private/CookerSettings.cpp
--- a/Engine/Source/Developer/DeveloperToolSettings/Private/CookerSettings.cpp
+++ b/Engine/Source/Developer/DeveloperToolSettings/Private/CookerSettings.cpp
@@ -26,6 +26,8 @@ UCookerSettings::UCookerSettings(const FObjectInitializer& ObjectInitializer)
 	
 	DefaultASTCQualityBySpeed = 2; // Medium preset
 	DefaultASTCQualityBySize = 3;  // 6x6
+	DefaultASTCQualityBySizeHQ = 4; // 4x4
+
 	{
 		static IConsoleVariable* ASTCTextureCompressorCVar = IConsoleManager::Get().FindConsoleVariable(TEXT("cook.ASTCTextureCompressor"));
 		DefaultASTCCompressor = (ASTCTextureCompressorCVar && ASTCTextureCompressorCVar->GetInt() != 0) ? ETextureFormatASTCCompressor::Arm : ETextureFormatASTCCompressor::IntelISPC;
diff --git a/Engine/Source/Developer/TextureFormatASTC/Private/TextureFormatASTC.cpp b/Engine/Source/Developer/TextureFormatASTC/Private/TextureFormatASTC.cpp
--- a/Engine/Source/Developer/TextureFormatASTC/Private/TextureFormatASTC.cpp
+++ b/Engine/Source/Developer/TextureFormatASTC/Private/TextureFormatASTC.cpp
@@ -55,6 +55,26 @@ static FAutoConsoleVariableRef CVarASTCCompressor(
 #define MAX_QUALITY_BY_SIZE 4
 #define MAX_QUALITY_BY_SPEED 3
 
+/**
+
+"Quality" in this file is ETextureCompressionQuality-1
+
+so a "3" here == High == 6x6
+
+enum ETextureCompressionQuality : int
+{
+	TCQ_Default = 0		UMETA(DisplayName="Default"),
+	TCQ_Lowest = 1		UMETA(DisplayName="Lowest (ASTC 12x12)"),
+	TCQ_Low = 2			UMETA(DisplayName="Low (ASTC 10x10)"),
+	TCQ_Medium = 3		UMETA(DisplayName="Medium (ASTC 8x8)"),
+	TCQ_High= 4			UMETA(DisplayName="High (ASTC 6x6)"),
+	TCQ_Highest = 5		UMETA(DisplayName="Highest (ASTC 4x4)"),
+	TCQ_MAX,
+};
+
+
+**/
+
 
 DEFINE_LOG_CATEGORY_STATIC(LogTextureFormatASTC, Log, All);
 
@@ -144,12 +164,11 @@ static int32 GetDefaultCompressionBySizeValue(FCbObjectView InFormatConfigOverri
 	}
 	else
 	{
-		// default of 0 == 12x12 ?
-		// BaseEngine.ini sets DefaultASTCQualityBySize to 3 == 6x6
+		// default of 3 == 6x6
 
 		auto GetCompressionModeValue = []() {
 			// start at default quality, then lookup in .ini file
-			int32 CompressionModeValue = 0;
+			int32 CompressionModeValue = 3;
 			GConfig->GetInt(TEXT("/Script/UnrealEd.CookerSettings"), TEXT("DefaultASTCQualityBySize"), CompressionModeValue, GEngineIni);
 	
 			FParse::Value(FCommandLine::Get(), TEXT("-astcqualitybysize="), CompressionModeValue);
@@ -163,6 +182,38 @@ static int32 GetDefaultCompressionBySizeValue(FCbObjectView InFormatConfigOverri
 	}
 }
 
+static int32 GetDefaultCompressionBySizeValueHQ(FCbObjectView InFormatConfigOverride)
+{
+	// this is code duped between TextureFormatASTC and TextureFormatISPC
+	if (InFormatConfigOverride)
+	{
+		// If we have an explicit format config, then use it directly
+		FCbFieldView FieldView = InFormatConfigOverride.FindView("DefaultASTCQualityBySizeHQ");
+		checkf(FieldView.HasValue(), TEXT("Missing DefaultASTCQualityBySizeHQ key from FormatConfigOverride"));
+		int32 CompressionModeValue = FieldView.AsInt32();
+		checkf(!FieldView.HasError(), TEXT("Failed to parse DefaultASTCQualityBySizeHQ value from FormatConfigOverride"));
+		return CompressionModeValue;
+	}
+	else
+	{
+		// default of 4 == 4x4
+
+		auto GetCompressionModeValue = []() {
+			// start at default quality, then lookup in .ini file
+			int32 CompressionModeValue = 4;
+			GConfig->GetInt(TEXT("/Script/UnrealEd.CookerSettings"), TEXT("DefaultASTCQualityBySizeHQ"), CompressionModeValue, GEngineIni);
+	
+			FParse::Value(FCommandLine::Get(), TEXT("-astcqualitybysizehq="), CompressionModeValue);
+			
+			return FMath::Min<uint32>(CompressionModeValue, MAX_QUALITY_BY_SIZE);
+		};
+
+		static int32 CompressionModeValue = GetCompressionModeValue();
+
+		return CompressionModeValue;
+	}
+}
+
 static int32 GetDefaultCompressionBySpeedValue(FCbObjectView InFormatConfigOverride)
 {
 	if (InFormatConfigOverride)
@@ -177,11 +228,11 @@ static int32 GetDefaultCompressionBySpeedValue(FCbObjectView InFormatConfigOverr
 	else
 	{
 
-		// default of 0 == "fastest"
+		// default of 2 == ASTCENC_PRE_MEDIUM
 
 		auto GetCompressionModeValue = []() {
 			// start at default quality, then lookup in .ini file
-			int32 CompressionModeValue = 0;
+			int32 CompressionModeValue = 2;
 			GConfig->GetInt(TEXT("/Script/UnrealEd.CookerSettings"), TEXT("DefaultASTCQualityBySpeed"), CompressionModeValue, GEngineIni);
 	
 			FParse::Value(FCommandLine::Get(), TEXT("-astcqualitybyspeed="), CompressionModeValue);
@@ -203,20 +254,20 @@ static EPixelFormat GetQualityFormat(const FTextureBuildSettings& BuildSettings)
 	int32 OverrideSizeValue= BuildSettings.CompressionQuality;
 
 	bool bIsNormalMap = IsNormalMapFormat(BuildSettings.TextureFormatName);
-	bool bIsHQ = BuildSettings.TextureFormatName == GTextureFormatNameASTC_RGBA_HQ;
-	bool bHDRFormat = BuildSettings.TextureFormatName == GTextureFormatNameASTC_RGB_HDR;
 
 	if ( bIsNormalMap )
 	{
+		// normal map hard coded to always use 6x6 currently
+		//	ignores per-texture quality
+
 		if ( BuildSettings.TextureFormatName == GTextureFormatNameASTC_NormalRG_Precise )
 		{
 			return PF_ASTC_6x6_NORM_RG;
 		}
-		return PF_ASTC_6x6;
-	}
-	else if ( bIsHQ )
-	{
-		return PF_ASTC_4x4;
+		else
+		{
+			return PF_ASTC_6x6;
+		}
 	}
 	else if (BuildSettings.bVirtualStreamable)
 	{
@@ -224,12 +275,27 @@ static EPixelFormat GetQualityFormat(const FTextureBuildSettings& BuildSettings)
 	}
 
 	// CompressionQuality value here is ETextureCompressionQuality minus 1
+	
+	bool bIsHQ = BuildSettings.TextureFormatName == GTextureFormatNameASTC_RGBA_HQ;
+	bool bHDRFormat = BuildSettings.TextureFormatName == GTextureFormatNameASTC_RGB_HDR;
+	
+	if ( OverrideSizeValue < 0 )
+	{
+		if ( bIsHQ )
+		{
+			OverrideSizeValue = GetDefaultCompressionBySizeValueHQ(InFormatConfigOverride);
+		}
+		else
+		{
+			OverrideSizeValue = GetDefaultCompressionBySizeValue(InFormatConfigOverride);
+		}
+	}
 
 	// convert to a string
 	EPixelFormat Format = PF_Unknown;
 	if (bHDRFormat)
 	{
-		switch (OverrideSizeValue >= 0 ? OverrideSizeValue : GetDefaultCompressionBySizeValue(InFormatConfigOverride))
+		switch (OverrideSizeValue)
 		{
 			case 0:	Format = PF_ASTC_12x12_HDR; break;
 			case 1:	Format = PF_ASTC_10x10_HDR; break;
@@ -241,7 +307,7 @@ static EPixelFormat GetQualityFormat(const FTextureBuildSettings& BuildSettings)
 	}
 	else
 	{
-		switch (OverrideSizeValue >= 0 ? OverrideSizeValue : GetDefaultCompressionBySizeValue(InFormatConfigOverride))
+		switch (OverrideSizeValue)
 		{
 			case 0:	Format = PF_ASTC_12x12; break;
 			case 1:	Format = PF_ASTC_10x10; break;
@@ -515,6 +581,7 @@ public:
 		FCbWriter Writer;
 		Writer.BeginObject("TextureFormatASTCSettings");
 		Writer.AddInteger("DefaultASTCQualityBySize", GetDefaultCompressionBySizeValue(FCbObjectView()));
+		Writer.AddInteger("DefaultASTCQualityBySizeHQ", GetDefaultCompressionBySizeValueHQ(FCbObjectView()));
 		Writer.AddInteger("DefaultASTCQualityBySpeed", GetDefaultCompressionBySpeedValue(FCbObjectView()));
 		Writer.EndObject();
 		return Writer.Save().AsObject();
diff --git a/Engine/Source/Developer/TextureFormatIntelISPCTexComp/Private/TextureFormatIntelISPCTexComp.cpp b/Engine/Source/Developer/TextureFormatIntelISPCTexComp/Private/TextureFormatIntelISPCTexComp.cpp
--- a/Engine/Source/Developer/TextureFormatIntelISPCTexComp/Private/TextureFormatIntelISPCTexComp.cpp
+++ b/Engine/Source/Developer/TextureFormatIntelISPCTexComp/Private/TextureFormatIntelISPCTexComp.cpp
@@ -467,6 +467,7 @@ static EPixelFormat GetASTCQualityFormat(int& BlockWidth, int& BlockHeight, cons
 	}
 	else if ( bIsHQ || BuildSettings.bVirtualStreamable )
 	{
+		// NOTE: different than TextureFormatASTC which configurable HQ block size
 		BlockWidth = BlockHeight = 4;
 		return PF_ASTC_4x4;		
 	}
