TextureFormatASTC : stop loading ImageWrapper and ISPC that aren't actually used

diff --git a/Engine/Source/Developer/TextureFormatASTC/Private/TextureFormatASTC.cpp b/Engine/Source/Developer/TextureFormatASTC/Private/TextureFormatASTC.cpp
--- a/Engine/Source/Developer/TextureFormatASTC/Private/TextureFormatASTC.cpp
+++ b/Engine/Source/Developer/TextureFormatASTC/Private/TextureFormatASTC.cpp
@@ -10,8 +10,6 @@
 #include "Misc/FileHelper.h"
 #include "Misc/Guid.h"
 #include "Misc/Paths.h"
-#include "IImageWrapper.h"
-#include "IImageWrapperModule.h"
 #include "ImageCore.h"
 #include "Interfaces/ITextureFormat.h"
 #include "Interfaces/ITextureFormatModule.h"
@@ -598,11 +596,22 @@ class FTextureFormatASTC : public ITextureFormat
 {
 public:
 	FTextureFormatASTC()
-	:	IntelISPCTexCompFormat(*FModuleManager::LoadModuleChecked<ITextureFormatModule>(TEXT("TextureFormatIntelISPCTexComp")).GetTextureFormat()),
-		ImageWrapperModule(FModuleManager::LoadModuleChecked<IImageWrapperModule>(FName("ImageWrapper")))
 	{
 		// LoadModule has to be done on Main thread
 		// can't be done on-demand in the Compress call
+		
+#if SUPPORTS_ISPC_ASTC
+		const bool bAllowTogglingISPCAfterStartup = false; // option
+
+		if(GASTCCompressor == 0 || bAllowTogglingISPCAfterStartup)
+		{
+			ITextureFormatModule * IntelISPCTexCompModule = FModuleManager::LoadModulePtr<ITextureFormatModule>(TEXT("TextureFormatIntelISPCTexComp"));
+			if ( IntelISPCTexCompModule )
+			{
+				IntelISPCTexCompFormat = IntelISPCTexCompModule->GetTextureFormat();
+			}
+		}
+#endif
 	}
 
 	static FGuid GetDecodeBuildFunctionVersionGuid()
@@ -622,9 +631,9 @@ public:
 	virtual bool AllowParallelBuild() const override
 	{
 #if SUPPORTS_ISPC_ASTC
-		if(GASTCCompressor == 0)
+		if(GASTCCompressor == 0 && IntelISPCTexCompFormat)
 		{
-			return IntelISPCTexCompFormat.AllowParallelBuild();
+			return IntelISPCTexCompFormat->AllowParallelBuild();
 		}
 #endif
 		return true;
@@ -632,9 +641,9 @@ public:
 	virtual FName GetEncoderName(FName Format) const override
 	{
 #if SUPPORTS_ISPC_ASTC
-		if (GASTCCompressor == 0)
+		if (GASTCCompressor == 0 && IntelISPCTexCompFormat)
 		{
-			return IntelISPCTexCompFormat.GetEncoderName(Format);
+			return IntelISPCTexCompFormat->GetEncoderName(Format);
 		}
 #endif
 		static const FName ASTCName("ArmASTC");
@@ -644,9 +653,9 @@ public:
 	virtual FCbObject ExportGlobalFormatConfig(const FTextureBuildSettings& BuildSettings) const override
 	{
 #if SUPPORTS_ISPC_ASTC
-		if(GASTCCompressor == 0)
+		if(GASTCCompressor == 0 && IntelISPCTexCompFormat)
 		{
-			return IntelISPCTexCompFormat.ExportGlobalFormatConfig(BuildSettings);
+			return IntelISPCTexCompFormat->ExportGlobalFormatConfig(BuildSettings);
 		}
 #endif
 		FCbWriter Writer;
@@ -665,11 +674,11 @@ public:
 	) const override
 	{
 #if SUPPORTS_ISPC_ASTC
-		if(GASTCCompressor == 0)
+		if(GASTCCompressor == 0 && IntelISPCTexCompFormat)
 		{
 			// set high bit so version numbers of ISPC and ASTC don't overlap :
 			check( BASE_ASTC_FORMAT_VERSION < 0x80 );
-			return 0x80 | IntelISPCTexCompFormat.GetVersion(Format,BuildSettings);
+			return 0x80 | IntelISPCTexCompFormat->GetVersion(Format,BuildSettings);
 		}
 #endif
 
@@ -679,9 +688,9 @@ public:
 	virtual FString GetDerivedDataKeyString(const FTextureBuildSettings& InBuildSettings, int32 InMipCount, const FIntVector3& InMip0Dimensions) const override
 	{
 #if SUPPORTS_ISPC_ASTC
-		if(GASTCCompressor == 0)
+		if(GASTCCompressor == 0 && IntelISPCTexCompFormat)
 		{
-			return IntelISPCTexCompFormat.GetDerivedDataKeyString(InBuildSettings, InMipCount, InMip0Dimensions);
+			return IntelISPCTexCompFormat->GetDerivedDataKeyString(InBuildSettings, InMipCount, InMip0Dimensions);
 		}
 #endif
 
@@ -809,7 +818,7 @@ public:
 		) const override
 	{
 #if SUPPORTS_ISPC_ASTC
-		if(GASTCCompressor == 0)
+		if(GASTCCompressor == 0 && IntelISPCTexCompFormat)
 		{
 			UE_CALL_ONCE( [&](){
 				UE_LOG(LogTextureFormatASTC, Display, TEXT("TextureFormatASTC using ISPC"))
@@ -817,7 +826,7 @@ public:
 
 			// Route ASTC compression work to the ISPC module instead.
 			// note: ISPC can't do HDR, will throw an error
-			return IntelISPCTexCompFormat.CompressImage(InImage, BuildSettings, InMip0Dimensions, InMip0NumSlicesNoDepth, InMipIndex, InMipCount, DebugTexturePathName, bImageHasAlphaChannel, OutCompressedImage);
+			return IntelISPCTexCompFormat->CompressImage(InImage, BuildSettings, InMip0Dimensions, InMip0NumSlicesNoDepth, InMipIndex, InMipCount, DebugTexturePathName, bImageHasAlphaChannel, OutCompressedImage);
 		}
 #endif
 
@@ -831,9 +840,7 @@ public:
 	}
 
 private:
-	const ITextureFormat& IntelISPCTexCompFormat;
-	
-	IImageWrapperModule& ImageWrapperModule;
+	const ITextureFormat * IntelISPCTexCompFormat = nullptr;
 };
 
 /**
diff --git a/Engine/Source/Developer/TextureFormatASTC/TextureFormatASTC.Build.cs b/Engine/Source/Developer/TextureFormatASTC/TextureFormatASTC.Build.cs
--- a/Engine/Source/Developer/TextureFormatASTC/TextureFormatASTC.Build.cs
+++ b/Engine/Source/Developer/TextureFormatASTC/TextureFormatASTC.Build.cs
@@ -17,7 +17,6 @@ public class TextureFormatASTC : ModuleRules
 			"Core",
 			"DerivedDataCache", // for FCacheBucket
 			"ImageCore",
-			"ImageWrapper",
 			"TextureBuild",
 			"TextureFormatIntelISPCTexComp",
 			"astcenc",
