Set engine default astc version to 5.0.1

diff --git a/Engine/Config/Android/BaseAndroidEngine.ini b/Engine/Config/Android/BaseAndroidEngine.ini
--- a/Engine/Config/Android/BaseAndroidEngine.ini
+++ b/Engine/Config/Android/BaseAndroidEngine.ini
@@ -31,3 +31,8 @@ r.Substrate.Glints=0
 r.Substrate.ClosuresPerPixel=1
 ; Volumetric Clouds are not supported by mobile renderer, skip materials/global shaders compilation.
 r.VolumetricCloud.Support=0
+
+[/Script/AndroidRuntimeSettings.AndroidRuntimeSettings]
+; Default to latest astc version - projects can manually specify this to retain older versions to prevent
+; patches.
+ASTCVersion=501
\ No newline at end of file
diff --git a/Engine/Source/Developer/TextureFormatASTC/Private/TextureFormatASTC.cpp b/Engine/Source/Developer/TextureFormatASTC/Private/TextureFormatASTC.cpp
--- a/Engine/Source/Developer/TextureFormatASTC/Private/TextureFormatASTC.cpp
+++ b/Engine/Source/Developer/TextureFormatASTC/Private/TextureFormatASTC.cpp
@@ -834,8 +834,9 @@ public:
 			ASTCSuffix << TEXT("RDO_") << InBuildSettings.OodleRDO;
 		}
 
-		// only add in the version if we aren't the first version (4.2.0)
-		if (InBuildSettings.AstcEncVersion != NAME_None)
+		// only add in the version if we aren't the first version (4.2.0). None defaults to first version.
+		if (InBuildSettings.AstcEncVersion != NAME_None &&
+			InBuildSettings.AstcEncVersion != SupportedAstcEncVersions[0])
 		{
 			ASTCSuffix << TEXT("V_") << InBuildSettings.AstcEncVersion;
 		}
